LLM: Change tracking of allocations in LowLevelMemTracker to allow TagSets to always be compiled in without a large increase in memory cost, by creating a middle layer of FAllocationGroups that contain the active TagData for each active TagSet. Allocations now include a single integer identifying their allocation group rather than needing a separate integer per ELLMTagSet.
#rnx

#Synced-CL 49735919

diff --git a/Engine/Source/Runtime/Core/Private/HAL/LLM/ActiveTags.cpp b/Engine/Source/Runtime/Core/Private/HAL/LLM/ActiveTags.cpp
new file mode 100644
--- /dev/null
+++ b/Engine/Source/Runtime/Core/Private/HAL/LLM/ActiveTags.cpp
@@ -0,0 +1,152 @@
+// Copyright Epic Games, Inc. All Rights Reserved.
+
+#include "HAL/LLM/ActiveTags.h"
+
+#if ENABLE_LOW_LEVEL_MEM_TRACKER
+
+namespace UE::LLMPrivate
+{
+
+FActiveTags::FActiveTags()
+{
+	const int32 GlobalNum = GetGlobalNum();
+	for (int32 Index = 0; Index < GlobalNum; ++Index)
+	{
+		Values[Index] = nullptr;
+	}
+}
+
+FActiveTags::FActiveTags(const FActiveTags& Other)
+{
+	*this = Other;
+}
+
+FActiveTags::FActiveTags(FActiveTags&& Other)
+{
+	*this = (FActiveTags&&)Other;
+}
+
+FActiveTags& FActiveTags::operator=(const FActiveTags& Other)
+{
+	const int32 GlobalNum = GetGlobalNum();
+	for (int32 Index = 0; Index < GlobalNum; ++Index)
+	{
+		Values[Index] = Other.Values[Index];
+	}
+	return *this;
+}
+
+FActiveTags& FActiveTags::operator=(FActiveTags&& Other)
+{
+	const int32 GlobalNum = GetGlobalNum();
+	for (int32 Index = 0; Index < GlobalNum; ++Index)
+	{
+		Values[Index] = Other.Values[Index];
+	}
+	return *this;
+}
+
+const FTagData*& FActiveTags::operator[](int32 Index)
+{
+	LLMCheck(0 <= Index);
+	LLMCheck(Index < GetGlobalNum());
+	return Values[Index];
+}
+
+const FTagData* const& FActiveTags::operator[](int32 Index) const
+{
+	LLMCheck(0 <= Index);
+	LLMCheck(Index < GetGlobalNum());
+	return Values[Index];
+}
+
+int32 FActiveTags::Num() const
+{
+	return GetGlobalNum();
+}
+
+bool FActiveTags::operator==(const FActiveTags& Other) const
+{
+	const int32 GlobalNum = GetGlobalNum();
+	for (int32 Index = 0; Index < GlobalNum; ++Index)
+	{
+		if (Values[Index] != Other.Values[Index])
+		{
+			return false;
+		}
+	}
+	return true;
+}
+
+const FTagData* FActiveTags::GetSystemsTagData() const
+{
+	LLMCheck(GetGlobalNum() > 0);
+	return Values[0];
+}
+
+void FActiveTags::SetSystemsTagData(const FTagData* TagData)
+{
+	LLMCheck(GetGlobalNum() > 0);
+	Values[0] = TagData;
+}
+
+const FTagData** FActiveTags::begin()
+{
+	return Values;
+}
+
+const FTagData* const* FActiveTags::begin() const
+{
+	return Values;
+}
+
+const FTagData* const* FActiveTags::end() const
+{
+	return Values + GetGlobalNum();
+}
+
+int32 FActiveTags::GetGlobalNum()
+{
+	return GetGlobalNumStorage();
+}
+
+void FActiveTags::SetGlobalNum(int32 NewNum)
+{
+	GetGlobalNumStorage() = NewNum;
+}
+
+FActiveTags FActiveTags::ConvertToPostCommandlineBootstrap(const FActiveTags& NewDefaults)
+{
+	FActiveTags Result;
+	Result[0] = (*this)[0];
+	for (int32 Index = 1; Index < GetGlobalNum(); ++Index)
+	{
+		Result[Index] = NewDefaults[Index];
+	}
+	return Result;
+}
+
+int32& FActiveTags::GetGlobalNumStorage()
+{
+	static int32 GlobalNum = 1;
+	return GlobalNum;
+}
+
+uint32 GetTypeHash(const FActiveTags& Values)
+{
+	using namespace UE::LLMPrivate;
+
+	// FTagData are never deleted, so we can use their pointer values
+	constexpr uint32 Initial = 101;
+	constexpr uint32 Multiplier = 103;
+	uint32 Result = Initial;
+	for (const FTagData* Ptr : Values)
+	{
+		Result = Result * Multiplier + ::GetTypeHash(Ptr);
+	}
+	return Result;
+}
+
+} // namespace UE::LLMPrivate
+
+#endif // ENABLE_LOW_LEVEL_MEM_TRACKER
\ No newline at end of file
diff --git a/Engine/Source/Runtime/Core/Private/HAL/LLM/ActiveTags.h b/Engine/Source/Runtime/Core/Private/HAL/LLM/ActiveTags.h
new file mode 100644
--- /dev/null
+++ b/Engine/Source/Runtime/Core/Private/HAL/LLM/ActiveTags.h
@@ -0,0 +1,65 @@
+// Copyright Epic Games, Inc. All Rights Reserved.
+
+#pragma once
+
+#include "HAL/LowLevelMemTracker.h"
+
+#if ENABLE_LOW_LEVEL_MEM_TRACKER
+
+#include "Templates/TypeHash.h"
+
+namespace UE::LLMPrivate
+{
+
+/**
+ * An embedded array of FTagData* that stores the active FTagData* for each of the FTagSet in LLM's global ordered
+ * list of FTagSets. These arrays are snapshotted at the time of an allocation and then are stored as the 
+ * key for the AllocationGroup shared by all allocations with that set tags.
+ * The size in memory of the array is set at compile time, equal to NumPossibleTagSets*sizeof(Pointer).
+ * The accessible length of the array is limited to NumActiveTagSets; the remaining values past the NumActiveTagSets
+ * index are uninitialized, and are not read/written publicly or privately, with the exception of LLM bootstrapping.
+ */
+class FActiveTags
+{
+public:
+	FActiveTags();
+	FActiveTags(const FActiveTags& Other);
+	FActiveTags(FActiveTags&& Other);
+	FActiveTags& operator=(const FActiveTags& Other);
+	FActiveTags& operator=(FActiveTags&& Other);
+
+	const FTagData*& operator[](int32 Index);
+	const FTagData*const& operator[](int32 Index) const;
+	int32 Num() const;
+
+	bool operator==(const FActiveTags& Other) const;
+
+	/** Get/Set the TagData stored for ELLMTagSet::None in this array. Non-null for initialized FActiveTags. */
+	const FTagData* GetSystemsTagData() const;
+	void SetSystemsTagData(const FTagData* TagData);
+
+	[[nodiscard]] const FTagData** begin();
+	[[nodiscard]] const FTagData* const* begin() const;
+	[[nodiscard]] const FTagData* const* end() const;
+
+	static int32 GetGlobalNum();
+	static void SetGlobalNum(int32 NewNum);
+
+	FActiveTags ConvertToPostCommandlineBootstrap(const FActiveTags& NewDefaults);
+
+private:
+	static int32& GetGlobalNumStorage();
+
+	/**
+	 * Values in the array. The used length is <= compile time length, and is stored in GetGlobalNum().
+	 * Before ActiveSets are parsed from commandline it is 1, after CommandLineBootstrapping it becomes constant and is
+	 * equal to the number of active sets.
+	 */
+	const FTagData* Values[(uint8) ELLMTagSet::Max];
+};
+
+extern uint32 GetTypeHash(const UE::LLMPrivate::FActiveTags& Values);
+
+} // namespace UE::LLMPrivate
+
+#endif // ENABLE_LOW_LEVEL_MEM_TRACKER
\ No newline at end of file
diff --git a/Engine/Source/Runtime/Core/Private/HAL/LLM/AllocationGroup.cpp b/Engine/Source/Runtime/Core/Private/HAL/LLM/AllocationGroup.cpp
new file mode 100644
--- /dev/null
+++ b/Engine/Source/Runtime/Core/Private/HAL/LLM/AllocationGroup.cpp
@@ -0,0 +1,206 @@
+// Copyright Epic Games, Inc. All Rights Reserved.
+
+#include "HAL/LLM/AllocationGroup.h"
+
+#if ENABLE_LOW_LEVEL_MEM_TRACKER
+
+namespace UE::LLMPrivate
+{
+
+FAllocationGroup::FAllocationGroup(int32 InIndexInAllocationGroups)
+	: IndexInAllocationGroups(InIndexInAllocationGroups)
+	, UpdatesSinceLastReference(0)
+	, Initialized(0)
+{
+}
+
+void FAllocationGroup::Initialize(const FActiveTags& InActiveTags)
+{
+	LLMCheck(InActiveTags[0] != nullptr); // ELLMTagSet::None should always be set
+
+	ActiveTags = InActiveTags;
+	for (const int64& Size : TrackerSizes)
+	{
+		LLMCheck(Size == 0);
+	}
+	LLMCheck(UpdatesSinceLastReference == 0);
+	LLMCheck(ActiveAllocs == 0);
+	LLMCheck(ThreadRefCount.load(std::memory_order_relaxed) == 0);
+
+	Initialized = 1;
+}
+
+void FAllocationGroup::Reset()
+{
+	Initialized = 0;
+
+	ActiveTags = FActiveTags();
+	for (int64& Size : TrackerSizes)
+	{
+		Size = 0;
+	}
+	UpdatesSinceLastReference = 0;
+	ActiveAllocs = 0;
+	ThreadRefCount.store(0, std::memory_order_relaxed);
+}
+
+bool FAllocationGroup::IsInitialized() const
+{
+	return Initialized != 0;
+}
+
+int32 FAllocationGroup::GetIndex() const
+{
+	return IndexInAllocationGroups;
+}
+
+const FActiveTags& FAllocationGroup::GetActiveTags() const
+{
+	return ActiveTags;
+}
+
+void FAllocationGroup::ConvertToPostCommandlineBootstrap(const FActiveTags& NewDefaults)
+{
+	ActiveTags = ActiveTags.ConvertToPostCommandlineBootstrap(NewDefaults);
+}
+
+void FAllocationGroup::AddSizeAndAllocReferences(ELLMTracker Tracker, int64 DeltaSize, int32 DeltaAllocs,
+	EPruningLevel PruningLevel)
+{
+	// Note that ActiveAllocs may sometimes go negative, because we are not taking accurate snapshots of it from
+	// all threads, see the comment in FLLMTracker::FetchAndClearTagSizes.
+	ActiveAllocs += DeltaAllocs;
+	TrackerSizes[static_cast<uint8>(Tracker)] += DeltaSize;
+	if (PruningLevel > EPruningLevel::None)
+	{
+		UpdatesSinceLastReference = 0;
+	}
+}
+
+int64 FAllocationGroup::GetSize(ELLMTracker Tracker) const
+{
+	return TrackerSizes[static_cast<uint8>(Tracker)];
+}
+
+int32 FAllocationGroup::GetUpdatesSinceLastReference() const
+{
+	return static_cast<int32>(UpdatesSinceLastReference);
+}
+
+void FAllocationGroup::IncrementUpdatesSinceLastReference()
+{
+	// UpdatesSinceLastReference is a 16-bit integer to save space; don't let it overflow.
+	// Our usage of it does not need an always-exact count, so just clamp it to max value if it reaches max value.
+	if (UpdatesSinceLastReference < 0xffff)
+	{
+		++UpdatesSinceLastReference;
+	}
+}
+
+void FAllocationGroup::ResetUpdatesSinceLastReference()
+{
+	UpdatesSinceLastReference = 0;
+}
+
+void FAllocationGroup::AddRefThread()
+{
+	ThreadRefCount.fetch_add(1, std::memory_order_relaxed);
+}
+
+void FAllocationGroup::ReleaseThread()
+{
+	int32 OldValue = ThreadRefCount.fetch_add(-1, std::memory_order_relaxed);
+	LLMCheck(OldValue >= 1);
+}
+
+bool FAllocationGroup::IsReferenced() const
+{
+	return ActiveAllocs != 0 || ThreadRefCount.load(std::memory_order_relaxed) != 0;
+}
+
+void FAllocationGroupTrackingData::ConditionalInitialize(FAllocationGroup* InGroup, bool bReferenceTrack)
+{
+	if (Group)
+	{
+		return;
+	}
+	LLMCheck(InGroup != nullptr);
+	Group = InGroup;
+	if (bReferenceTrack)
+	{
+		Group->AddRefThread();
+	}
+	Size = 0;
+	ActiveAllocs = 0;
+	UpdatesSinceLastAllocationChange = 0;
+}
+
+void FAllocationGroupTrackingData::ConditionalReset(bool bReferenceTrack)
+{
+	if (!Group)
+	{
+		return;
+	}
+	if (bReferenceTrack)
+	{
+		Group->ReleaseThread();
+	}
+	Group = nullptr;
+}
+
+bool FAllocationGroupTrackingData::IsInitialized() const
+{
+	return Group != nullptr;
+}
+
+FAllocationGroup* FAllocationGroupTrackingData::GetGroup() const
+{
+	return Group;
+}
+
+int64 FAllocationGroupTrackingData::GetSize() const
+{
+	return Size;
+}
+
+int32 FAllocationGroupTrackingData::GetActiveAllocs() const
+{
+	return ActiveAllocs;
+}
+
+void FAllocationGroupTrackingData::TrackAllocOrFree(int64 DeltaSize, int32 DeltaActiveAllocs)
+{
+	Size += DeltaSize;
+	ActiveAllocs += DeltaActiveAllocs;
+	UpdatesSinceLastAllocationChange = 0;
+}
+
+void FAllocationGroupTrackingData::DetachSizeAndActiveAllocs(FAllocationGroupTrackingData& Receiver, EPruningLevel PruningLevel)
+{
+	if (!Receiver.Group)
+	{
+		Receiver.Group = Group;
+		Receiver.Size = 0;
+		LLMCheck(Receiver.GetActiveAllocs() == 0);
+		Receiver.ActiveAllocs = 0;
+	}
+	Receiver.Size += Size;
+	Receiver.ActiveAllocs += ActiveAllocs;
+
+	Size = 0;
+	ActiveAllocs = 0;
+	if (PruningLevel > EPruningLevel::None)
+	{
+		++UpdatesSinceLastAllocationChange;
+	}
+}
+
+int32 FAllocationGroupTrackingData::GetUpdatesSinceLastAllocationChange() const
+{
+	return UpdatesSinceLastAllocationChange;
+}
+
+} // namespace UE::LLMPrivate
+
+
+#endif // ENABLE_LOW_LEVEL_MEM_TRACKER
\ No newline at end of file
diff --git a/Engine/Source/Runtime/Core/Private/HAL/LLM/AllocationGroup.h b/Engine/Source/Runtime/Core/Private/HAL/LLM/AllocationGroup.h
new file mode 100644
--- /dev/null
+++ b/Engine/Source/Runtime/Core/Private/HAL/LLM/AllocationGroup.h
@@ -0,0 +1,169 @@
+// Copyright Epic Games, Inc. All Rights Reserved.
+
+#pragma once
+
+#include "HAL/LowLevelMemTracker.h"
+
+#if ENABLE_LOW_LEVEL_MEM_TRACKER
+
+#include "Containers/Array.h"
+#include "Containers/Map.h"
+#include "HAL/LLM/ActiveTags.h"
+#include "HAL/LLM/LLMPrivateDefines.h"
+#include "HAL/LLM/MemoryUtils.h"
+
+#include <atomic>
+
+namespace UE::LLMPrivate
+{
+
+/**
+ * Records the list of tags that were active when a pointer was allocated, shared by the allocationinfo for all
+ * pointers that had the same set of tags. AllocationGroups can be identified by their 4-byte index in a global
+ * array, allowing us to minimize the memory overheaded necessary on the AllocationInfo to update tag sizes when
+ * the allocation is freed. Allocation groups are discarded when no longer referenced by an active allocation (with
+ * a few exceptions).
+ * 
+ * Records the size and number of active allocations made under that set of tags. These values are not precise due
+ * to threading; when the values are updated we take snapshots of each thread's change in allocations outside of
+ * a critical section, and a race condition may cause an allocation at the coarse-snapshot time to be temporarily
+ * miscounted.
+ * 
+ * Sizes from all allocation groups with a given FTagData are summed to provide the size data for the TagData.
+ */
+class FAllocationGroup
+{
+public:
+	FAllocationGroup(int32 InIndexInAllocationGroups);
+
+	void Initialize(const FActiveTags& InActiveTags);
+	void Reset();
+
+	/** Initialized flag, used to implement a pool allocator for Allocation groups. */
+	bool IsInitialized() const;
+	/** Immutable data, this group's index in LLM's AllocationGroups. */
+	int32 GetIndex() const;
+	/**
+	 * (Mostly) immutable data: the ActiveTags that this group is tracking.
+	 * Changes only during CommandLineBootStrapping.
+	 */
+	const FActiveTags& GetActiveTags() const;
+
+	/** Modifies otherwise immutable data, call only when all threads are paused by the global LLM EnableLock. */
+	void ConvertToPostCommandlineBootstrap(const FActiveTags& NewDefaults);
+
+	/**
+	 * Update-thread only.
+	 * Adds the accumulated change and Size and ActiveAllocs that were recorded by a thread over an update period.
+	 * Also updates lifetime data for later pruning, if PruningLevel != EPruningLevel::None.
+	 */
+	void AddSizeAndAllocReferences(ELLMTracker Tracker, int64 DeltaSize, int32 DeltaAllocs, EPruningLevel PruningLevel);
+	/** Update-thread only. Returns the currently known size of allocations in this group. */
+	int64 GetSize(ELLMTracker Tracker) const;
+
+	/** Update-thread only. Get/Modify the number of updates since the group was last referenced, used in pruning. */
+	int32 GetUpdatesSinceLastReference() const;
+	void IncrementUpdatesSinceLastReference();
+	void ResetUpdatesSinceLastReference();
+
+	/**
+	 * Thread-safe increment of the reference counter on this. Reference count does not trigger destruction,
+	 * it is informational only, does not have to include all references, and is used by the pruning code.
+	 */
+	void AddRefThread();
+	/**
+	 * Thread-safe decrement of the reference counter on this. Reference count does not trigger destruction,
+	 * it is informational only, does not have to include all references, and is used by the pruning code.
+	 */
+	void ReleaseThread();
+	/**
+	 * Thread-safe report whether any threads or active allocations exist for *this.
+	 * Note that due to threading race conditions it is possible for our ActiveAllocs field to incorrectly hold
+	 * the value zero, but it is not possible for that to occur at the same time as ThreadRefCount == 0, so
+	 * the combined boolean value of ThreadRefCount != 0 || ActiveAllocs != 0 is always accurate. See the comment
+	 * in FLLMTracker::FetchAndClearTagSizes.
+	 */ 
+	bool IsReferenced() const;
+
+private:
+	FActiveTags ActiveTags;
+	int64 TrackerSizes[static_cast<uint8>(ELLMTracker::Max)] = { 0 };
+	int32 IndexInAllocationGroups = -1;
+	uint32 UpdatesSinceLastReference : 16; // Initialized to 0
+	uint32 Initialized : 1; // Initialized to 0
+	int32 ActiveAllocs = 0;
+	std::atomic<int32> ThreadRefCount{ 0 };
+};
+
+/**
+ * Data held by a thread or during update that records the change in number and size of allocations under
+ * an allocation group. All data is non-threadsafe, and should be guarded externally from simultaneous access
+ * on multiple threads.
+ */
+class FAllocationGroupTrackingData
+{
+public:
+	/**
+	 * Initialize the TrackingData if not already initialized. Used when adding/removing TrackingDatas
+	 * from threads and the update thread. bReferenceTrack is set to true when the TrackingData is stored
+	 * on in the thread state for a thread other than the UpdateThread and therefore we need to record
+	 * its reference to its AllocationGroup to prevent the group from being pruned.
+	 */
+	void ConditionalInitialize(FAllocationGroup* InGroup, bool bReferenceTrack = false);
+	void ConditionalReset(bool bReferenceTrack = false);
+	bool IsInitialized() const;
+
+	/** The AllocationGroup tracked by this tracking data. Non-null when this TrackingData is initialized. */
+	FAllocationGroup* GetGroup() const;
+
+	/** The postive or negative change in allocation size since the last time Detach was called. */
+	int64 GetSize() const;
+	/** The postive or negative change in the number of ActiveAllocs size since the last time Detach was called. */
+	int32 GetActiveAllocs() const;
+
+	/** Record the sum of a number of allocations or frees. */
+	void TrackAllocOrFree(int64 DeltaSize, int32 DeltaActiveAllocs);
+	/** Add the Size and ActiveAllocs from this into the Receiver, and optionally update reference data for pruning. */
+	void DetachSizeAndActiveAllocs(FAllocationGroupTrackingData& Receiver, EPruningLevel PruningLevel);
+
+	/**
+	 * Report the number of updates (calls to DetachSizeAndActiveAllocs) since the last time TrackAllocOrFree
+	 * was called; used to prune the TrackingData.
+	 */
+	int32 GetUpdatesSinceLastAllocationChange() const;
+
+private:
+	FAllocationGroup* Group = nullptr;
+	int64 Size = 0;
+	int32 ActiveAllocs = 0;
+	int32 UpdatesSinceLastAllocationChange = 0;
+};
+
+class FAllocationGroupArray : public TArray<FAllocationGroup*, FDefaultLLMAllocator>
+{
+public:
+	using TArray<FAllocationGroup*, FDefaultLLMAllocator>::TArray;
+};
+
+class FAllocationGroupSet : public TSet<FAllocationGroup*, DefaultKeyFuncs<FAllocationGroup*>, FDefaultSetLLMAllocator>
+{
+public:
+	using TSet<FAllocationGroup*, DefaultKeyFuncs<FAllocationGroup*>, FDefaultSetLLMAllocator>::TSet;
+};
+
+class FMapActiveTagsToAllocationGroup : public TMap<FActiveTags, FAllocationGroup*, FDefaultSetLLMAllocator>
+{
+public:
+	using TMap<FActiveTags, FAllocationGroup*, FDefaultSetLLMAllocator>::TMap;
+};
+
+class FMapActiveTagsToAllocationGroupTrackingData : public TMap<FActiveTags, FAllocationGroupTrackingData, FDefaultSetLLMAllocator>
+{
+public:
+	using TMap<FActiveTags, FAllocationGroupTrackingData, FDefaultSetLLMAllocator>::TMap;
+};
+
+} // namespace UE::LLMPrivate
+
+
+#endif // ENABLE_LOW_LEVEL_MEM_TRACKER
\ No newline at end of file
diff --git a/Engine/Source/Runtime/Core/Private/HAL/LLM/LLM.cpp b/Engine/Source/Runtime/Core/Private/HAL/LLM/LLM.cpp
--- a/Engine/Source/Runtime/Core/Private/HAL/LLM/LLM.cpp
+++ b/Engine/Source/Runtime/Core/Private/HAL/LLM/LLM.cpp
@@ -1,6 +1,7 @@
 // Copyright Epic Games, Inc. All Rights Reserved.
 
 #include "HAL/LowLevelMemTracker.h"
+
 #include "Containers/ContainerAllocationPolicies.h"
 
 #if ENABLE_LOW_LEVEL_MEM_TRACKER
@@ -9,6 +10,8 @@
 #include "Async/UniqueLock.h"
 #include "HAL/FileManager.h"
 #include "HAL/IConsoleManager.h"
+#include "HAL/LLM/ActiveTags.h"
+#include "HAL/LLM/AllocationGroup.h"
 #include "HAL/LLM/LLMPrivate.h"
 #include "HAL/LowLevelMemStats.h"
 #include "HAL/ThreadHeartBeat.h"
@@ -962,14 +965,8 @@ TAutoConsoleVariable<int32> CVarLLMCsvFlushEveryRow(TEXT("LLM.CsvFlushEveryRow")
 
 TAutoConsoleVariable<int32> CVarLLMHeaderMaxSize(
 	TEXT("LLM.LLMHeaderMaxSize"),
-#if LLM_ALLOW_ASSETS_TAGS
-	500000, // When using asset tags, you will have MANY more LLM titles since so many are auto generated.
-#elif LLM_ALLOW_UOBJECTCLASSES_TAGS
-	30000, // When using uobjects clases, you will have at least 20K+ LLM titles.
-#else
-	5000,
-#endif
-	TEXT("The maximum total number of characters allowed for all of the LLM titles")
+	-1,
+	TEXT("The maximum total number of characters allowed for all of the LLM TagNames, or -1 to calculate it from all ELLMTagSet::None tags.")
 );
 
 FAutoConsoleCommand LLMSnapshot(
@@ -1003,7 +1000,6 @@ FAutoConsoleCommand DumpLLM(
 		}
 
 		ELLMTagSet TagSet = ELLMTagSet::None;
-#if LLM_ALLOW_ASSETS_TAGS
 		if (FParse::Param(*Command, TEXT("Assets")))
 		{
 			TagSet = ELLMTagSet::Assets;
@@ -1012,12 +1008,10 @@ FAutoConsoleCommand DumpLLM(
 		{
 			TagSet = ELLMTagSet::AssetClasses;
 		}
-#elif LLM_ALLOW_UOBJECTCLASSES_TAGS
-		if (FParse::Param(*Command, TEXT("UObjectClasses")))
+		else if (FParse::Param(*Command, TEXT("UObjectClasses")))
 		{
 			TagSet = ELLMTagSet::UObjectClasses;
 		}
-#endif // LLM_ALLOW_UOBJECTCLASSES_TAGS
 
 		FLowLevelMemTracker::Get().DumpToLog(DumpFormat, &Ar, SizeParams, TagSet);
 	}));
@@ -1417,7 +1411,6 @@ extern const FName LLMGetUntaggedTagName(ELLMTagSet TagSet)
 		{
 			return TagName_Untagged;
 		}
-#if LLM_ALLOW_ASSETS_TAGS
 		case ELLMTagSet::Assets:
 		{
 			return TagName_UntaggedAsset;
@@ -1426,12 +1419,10 @@ extern const FName LLMGetUntaggedTagName(ELLMTagSet TagSet)
 		{
 			return TagName_UntaggedAssetClass;
 		}
-#elif LLM_ALLOW_UOBJECTCLASSES_TAGS
 		case ELLMTagSet::UObjectClasses:
 		{
 			return TagName_UntaggedUObjectClasses;
 		}
-#endif // LLM_ALLOW_UOBJECTCLASSES_TAGS
 		default:
 		{
 			return NAME_None;
@@ -1962,6 +1953,7 @@ void FLLMGlobals::BootstrapInitializeInner()
 	}
 
 	BootstrapTagDatas();
+	BootstrapAllocationGroups();
 	static_assert((uint8)ELLMTracker::Max == 2,
 		"You added a tracker, without updating FLowLevelMemTracker::BootstrapInitialize (and probably need to update macros)");
 	GetTracker(ELLMTracker::Platform)->SetTotalTags(FindOrAddTagData(ELLMTag::PlatformUntaggedTotal),
@@ -1999,6 +1991,7 @@ void FLLMGlobals::Clear()
 		TraceTagFilter = nullptr;
 	}
 
+	ClearAllocationGroups();
 	ClearTagDatas();
 	Allocator.Clear();
 	bFullyInitialized = false;
@@ -2184,10 +2177,10 @@ void FLLMGlobals::UpdateTags()
 	}
 	if (bNeedsResort)
 	{
-		// Prevent threads from reading their tags while we are mutating tags
+		// Prevent threads from reading their AllocationGroups or Tags while we are mutating tags.
 		for (int32 TrackerIndex = 0; TrackerIndex < static_cast<int32>(ELLMTracker::Max); TrackerIndex++)
 		{
-			GetTracker(static_cast<ELLMTracker>(TrackerIndex))->LockAllThreadTags(true);
+			GetTracker(static_cast<ELLMTracker>(TrackerIndex))->LockAllThreadsSharedData(true);
 		}
 
 		FTagDataArray* OldTagDatas;
@@ -2197,7 +2190,7 @@ void FLLMGlobals::UpdateTags()
 		for (int32 TrackerIndex = 0; TrackerIndex < static_cast<int32>(ELLMTracker::Max); TrackerIndex++)
 		{
 			GetTracker(static_cast<ELLMTracker>(TrackerIndex))->OnTagsResorted(*OldTagDatas);
-			GetTracker(static_cast<ELLMTracker>(TrackerIndex))->LockAllThreadTags(false);
+			GetTracker(static_cast<ELLMTracker>(TrackerIndex))->LockAllThreadsSharedData(false);
 		}
 
 		Allocator.Delete(OldTagDatas);
@@ -2412,48 +2405,74 @@ void FLLMGlobals::ProcessCommandLineInner(const TCHAR* CmdLine)
 	}
 	CSV_METADATA(TEXT("LLM"), TEXT("1"));
 
-	// PLATFORM_HAS_MULTITHREADED_PREMAIN: No need for a Write lock because we are not changing state to disabled.
-	// The other data we modify in this function is synchronized using other methods.
-	EnabledState = EEnabled::Enabled;
-	PRAGMA_DISABLE_DEPRECATION_WARNINGS;
-	bIsDisabled = false;
-	PRAGMA_ENABLE_DEPRECATION_WARNINGS;
-	bCsvWriterEnabled = bLocalCsvWriterEnabled;
-	bTraceWriterEnabled = bLocalTraceWriterEnabled;
-	FinishInitialise();
-
 	// activate tag sets (we ignore None set, it's always on)
 	FString SetList;
+	bool LocalActiveSets[static_cast<uint8>(ELLMTagSet::Max)] = { false };
+	LocalActiveSets[static_cast<uint8>(ELLMTagSet::None)] = true;
+
 	static_assert(static_cast<uint8>(ELLMTagSet::Max) == 4, "You added a tagset, without updating FLowLevelMemTracker::ProcessCommandLine");
 	if (FParse::Value(CmdLine, TEXT("LLMTAGSETS="), SetList, false /* bShouldStopOnSeparator */))
 	{
+#if LLM_ALLOW_NAMES_TAGS
 		TArray<FString> Sets;
 		SetList.ParseIntoArray(Sets, TEXT(","), true);
 		for (FString& Set : Sets)
 		{
-#if LLM_ALLOW_ASSETS_TAGS
 			if (Set == TEXT("Assets"))
 			{
-				ActiveSets[static_cast<int32>(ELLMTagSet::Assets)] = true;
+				LocalActiveSets[static_cast<int32>(ELLMTagSet::Assets)] = true;
 			}
 			else if (Set == TEXT("AssetClasses"))
 			{
-				ActiveSets[static_cast<int32>(ELLMTagSet::AssetClasses)] = true;
+				LocalActiveSets[static_cast<int32>(ELLMTagSet::AssetClasses)] = true;
 			}
-#elif LLM_ALLOW_UOBJECTCLASSES_TAGS
-			if (Set == TEXT("UObjectClasses"))
+			else if (Set == TEXT("UObjectClasses"))
 			{
-				ActiveSets[static_cast<int32>(ELLMTagSet::UObjectClasses)] = true;
+				LocalActiveSets[static_cast<int32>(ELLMTagSet::UObjectClasses)] = true;
 			}
+		}
 #else
-				// Asset tracking has a per-thread memory overhead, so we have a #define to completely disable it.
-				// Warn if we try to toggle it on at runtime but it is not defined.
-				UE_LOG(LogInit, Warning,
-					TEXT("Attempted to use LLM to track assets, but LLM_ALLOW_ASSETS_TAGS is not defined to 1. You will need to enable the define"));
+		// We do not support activating ELLMTagSets other than ELLMTagSet::None when not using LLM_ALLOW_NAMES_TAGS,
+		// because we need to assume we can use just the first TagData of the FActiveTags to
+		// map from ELLMTag -> FActiveTags, see FindOrAddAllocationGroup.
+		UE_LOG(LogInit, Warning,
+			TEXT("Attempted to use LLMTAGSETS, but LLM_ALLOW_NAMES_TAGS is not defined to 1. You will need to enable the define"));
 #endif
+	}
+	// These calls to FindOrAddTagData must occur outside of EnableStateLock to avoid deadlock due to
+	// recursive entry.
+	const FTagData* DefaultTags[static_cast<uint8>(ELLMTagSet::Max)] = { nullptr };
+	for (uint8 TagSetIndex = 0; TagSetIndex < static_cast<uint8>(ELLMTagSet::Max); ++TagSetIndex)
+	{
+		if (!LocalActiveSets[TagSetIndex])
+		{
+			continue;
 		}
+		DefaultTags[TagSetIndex] = FindOrAddDefaultTagData(static_cast<ELLMTagSet>(TagSetIndex));
 	}
 
+	{
+#if PLATFORM_HAS_MULTITHREADED_PREMAIN
+		// The EnableStateLock must be limited in scope because other code in the function
+		// allocates memory and would block on the readlock. Trying to take it around the write of the
+		// EnabledState and ActiveSets is sufficient; this will cause us to wait until threads already in
+		// OnLowLevelAlloc exit the function and clear their readlock, and will cause new calls to
+		// those functions to block while we're modifying IsEnabled and ActiveSets
+		// and call IsEnabled again after we release the lock.
+		FWriteScopeLock EnableStateLock(GetEnableStateLock());
+#endif
+		EnabledState = EEnabled::Enabled;
+		PRAGMA_DISABLE_DEPRECATION_WARNINGS;
+		bIsDisabled = false;
+		PRAGMA_ENABLE_DEPRECATION_WARNINGS;
+		FMemory::Memcpy(ActiveSets, LocalActiveSets, sizeof(ActiveSets));
+		OnActiveSetsInitialized(DefaultTags);
+	}
+
+	bCsvWriterEnabled = bLocalCsvWriterEnabled;
+	bTraceWriterEnabled = bLocalTraceWriterEnabled;
+	FinishInitialise();
+
 	if (FParse::Value(CmdLine, TEXT("LLMTraceTagFilter="), SetList, false /* bShouldStopOnSeparator */))
 	{
 		LLMCheck(TraceTagFilter == nullptr); // ProcessCommandLine() is expected to be called only once
@@ -2463,7 +2482,7 @@ void FLLMGlobals::ProcessCommandLineInner(const TCHAR* CmdLine)
 		{
 			if (!TraceTagFilter)
 			{
-				TraceTagFilter = Allocator.New<UE::LLMPrivate::FTagFilter>();
+				TraceTagFilter = Allocator.New<FTagFilter>();
 			}
 			for (FString& TagName : TagNames)
 			{
@@ -2473,7 +2492,7 @@ void FLLMGlobals::ProcessCommandLineInner(const TCHAR* CmdLine)
 			// TagData is supposed to be immutable once constructed, but we use bIsTraceReportable flag
 			// as a mutable member to quickly filter tags for the trace purposes.
 			UE::TSharedLock Lock(TagDataMutex);
-			for (UE::LLMPrivate::FTagData* TagData : (*TagDatas))
+			for (FTagData* TagData : (*TagDatas))
 			{
 				if (!TraceTagFilter->Contains(TagData->GetName()))
 				{
@@ -2504,6 +2523,80 @@ void FLLMGlobals::ProcessCommandLineInner(const TCHAR* CmdLine)
 	}
 }
 
+void FLLMGlobals::OnActiveSetsInitialized(const FTagData* DefaultTags[static_cast<uint8>(ELLMTagSet::Max)])
+{
+	// Called within write scope of EnableStateLock, #if PLATFORM_HAS_MULTITHREADED_PREMAIN
+	// Otherwise, called only before multithreading begins
+
+	// Satisfy our contract for thread synchronization of tags (even if not technically necessary because there
+	// are no threads currently): prevent threads from reading AllocationGroups or Tags while we are mutating
+	// AllocationGroups.
+	UE::TUniqueLock ScopeLockAllocationGroups(AllocationGroupsMutex);
+	for (int32 TrackerIndex = 0; TrackerIndex < static_cast<int32>(ELLMTracker::Max); TrackerIndex++)
+	{
+		GetTracker(static_cast<ELLMTracker>(TrackerIndex))->LockAllThreadsSharedData(true);
+	}
+
+	int32 NumActiveSets = 0;
+	for (bool IsActive : ActiveSets)
+	{
+		if (IsActive)
+		{
+			++NumActiveSets;
+		}
+	}
+	LLMCheck(NumActiveSets > 0);
+
+	if (NumActiveSets != FActiveTags::GetGlobalNum())
+	{
+		// Before we are allowed to call FActiveTags::SetGlobalNum, we have to remove all FActiveTags
+		// from the keys of TSets, because TSet keys are not allowed to mutate their TypeHash, and calling SetGlobalNum
+		// will mutate the TypeHash of every FActiveTags.
+		ActiveTagsToAllocationGroup->Empty();
+
+		TArray<TArray<FAllocationGroupTrackingData, FDefaultLLMAllocator>, FDefaultLLMAllocator>
+			ThreadStateAllocationGroupTrackingDatas[static_cast<int32>(ELLMTracker::Max)];
+		for (int32 TrackerIndex = 0; TrackerIndex < static_cast<int32>(ELLMTracker::Max); TrackerIndex++)
+		{
+			FLLMTracker* Tracker = GetTracker(static_cast<ELLMTracker>(TrackerIndex));
+			Tracker->UpdateThreads();
+			Tracker->BeginActiveSetsChange(ThreadStateAllocationGroupTrackingDatas[TrackerIndex]);
+		}
+
+		FActiveTags::SetGlobalNum(NumActiveSets);
+		FActiveTags NewDefaults;
+		int32 IndexInActiveTags = 0;
+		for (uint8 TagSetIndex = 0; TagSetIndex < static_cast<uint8>(ELLMTagSet::Max); ++TagSetIndex)
+		{
+			if (ActiveSets[TagSetIndex])
+			{
+				NewDefaults[IndexInActiveTags] = DefaultTags[TagSetIndex];
+				++IndexInActiveTags;
+			}
+		}
+
+		for (FAllocationGroup* AllocationGroup : (*AllocationGroups))
+		{
+			if (AllocationGroup->IsInitialized())
+			{
+				AllocationGroup->ConvertToPostCommandlineBootstrap(NewDefaults);
+				ActiveTagsToAllocationGroup->Add(AllocationGroup->GetActiveTags(), AllocationGroup);
+			}
+		}
+
+		for (int32 TrackerIndex = 0; TrackerIndex < static_cast<int32>(ELLMTracker::Max); TrackerIndex++)
+		{
+			GetTracker(static_cast<ELLMTracker>(TrackerIndex))->EndActiveSetsChange(NewDefaults,
+				ThreadStateAllocationGroupTrackingDatas[TrackerIndex]);
+		}
+	}
+
+	for (int32 TrackerIndex = 0; TrackerIndex < static_cast<int32>(ELLMTracker::Max); TrackerIndex++)
+	{
+		GetTracker(static_cast<ELLMTracker>(TrackerIndex))->LockAllThreadsSharedData(false);
+	}
+}
+
 uint64 FLLMGlobals::GetTotalTrackedMemoryInner(ELLMTracker Tracker)
 {
 	BootstrapInitialise();
@@ -2678,15 +2771,6 @@ bool FLLMGlobals::ShouldReduceThreadsInner()
 #endif
 }
 
-#if LLM_ALLOW_ASSETS_TAGS
-
-static bool IsAssetTagForAssets(ELLMTagSet Set)
-{
-	return Set == ELLMTagSet::Assets || Set == ELLMTagSet::AssetClasses;
-}
-
-#endif // LLM_ALLOW_ASSETS_TAGS
-
 void FLLMGlobals::RegisterCustomTag(int32 Tag, ELLMTagSet TagSet, const TCHAR* InDisplayName, FName StatName,
 	FName SummaryStatName, int32 ParentTag)
 {
@@ -2771,6 +2855,14 @@ void FLLMGlobals::BootstrapTagDatas()
 	}
 }
 
+void FLLMGlobals::BootstrapAllocationGroups()
+{
+	AllocationGroups = Allocator.New<FAllocationGroupArray>();
+	AllocationGroupsFreeList = Allocator.New<FAllocationGroupArray>();
+	AllocationGroupsUnreferencedSet = Allocator.New<FAllocationGroupSet>();
+	ActiveTagsToAllocationGroup = Allocator.New<FMapActiveTagsToAllocationGroup>();
+}
+
 void FLLMGlobals::FinishInitializeInner()
 {
 	if (bFullyInitialized)
@@ -2891,6 +2983,22 @@ void FLLMGlobals::ClearTagDatas()
 	TagDatas = nullptr;
 }
 
+void FLLMGlobals::ClearAllocationGroups()
+{
+	Allocator.Delete(ActiveTagsToAllocationGroup);
+	ActiveTagsToAllocationGroup = nullptr;
+	Allocator.Delete(AllocationGroupsFreeList);
+	AllocationGroupsFreeList = nullptr;
+	Allocator.Delete(AllocationGroupsUnreferencedSet);
+	AllocationGroupsUnreferencedSet = nullptr;
+	for (FAllocationGroup* AllocationGroup: *AllocationGroups)
+	{
+		Allocator.Delete(AllocationGroup);
+	}
+	Allocator.Delete(AllocationGroups);
+	AllocationGroups = nullptr;
+}
+
 void FLLMGlobals::RegisterTagDeclaration(FLLMTagDeclaration& TagDeclaration)
 {
 	TagDeclaration.ConstructUniqueName();
@@ -3503,6 +3611,81 @@ const FTagData* FLLMGlobals::FindTagData(FName TagName, ELLMTagSet TagSet, ETagR
 	}
 }
 
+const FTagData* FLLMGlobals::FindOrAddDefaultTagData(ELLMTagSet TagSet)
+{
+	switch (TagSet)
+	{
+	case ELLMTagSet::None:
+		return FindOrAddTagData(ELLMTag::Untagged, ETagReferenceSource::FunctionAPI);
+	case ELLMTagSet::Assets:
+		return FindOrAddTagData(TagName_UntaggedAsset, TagSet, false /* bIsStatTag */,
+			ETagReferenceSource::FunctionAPI);
+	case ELLMTagSet::AssetClasses:
+		return FindOrAddTagData(TagName_UntaggedAssetClass, TagSet, false /* bIsStatTag */,
+			ETagReferenceSource::FunctionAPI);
+	case ELLMTagSet::UObjectClasses:
+		return FindOrAddTagData(TagName_UntaggedUObjectClasses, TagSet, false /* bIsStatTag */,
+			ETagReferenceSource::FunctionAPI);
+	default:
+		return nullptr;
+	}
+}
+
+FAllocationGroup* FLLMGlobals::FindOrAddAllocationGroup(const FActiveTags& ActiveTags)
+{
+	// Caller must hold write lock on AllocationGroupsMutex.
+	// Caller must NOT be inside AllocationGroupsIndexMutex; this function may need to enter that mutex
+	FAllocationGroup*& AllocationGroup = ActiveTagsToAllocationGroup->FindOrAdd(ActiveTags);
+	if (!AllocationGroup)
+	{
+#if LLM_ALLOW_NAMES_TAGS // When !LLM_ALLOW_NAMES_TAGS, allocating AllocationGroups is more tightly constrained.
+		if (!AllocationGroupsFreeList->IsEmpty())
+		{
+			AllocationGroup = AllocationGroupsFreeList->Pop();
+		}
+		else
+		{
+			UE::TUniqueLock WriteScopeLockGlobalIndex(AllocationGroupsIndexMutex);
+			AllocationGroup = Allocator.New<FAllocationGroup>(AllocationGroups->Num());
+			AllocationGroups->Add(AllocationGroup);
+		}
+#else // LLM_ALLOW_NAMES_TAGS
+		// We do not support activating ELLMTagSets other than ELLMTagSet::None when not using FullTags,
+		// because we need to assume we can use just the first TagData of the FActiveTags to
+		// map from ELLMTag -> FActiveTags.
+		LLMCheck(ActiveTags.Num() == 1);
+
+		// We construct the each AllocationGroup to have an index equal to the integer value of its
+		// ELLMTagSet::None tag's EnumTag.
+		UE::TUniqueLock WriteScopeLockGlobalIndex(AllocationGroupsIndexMutex);
+		uint8 Index = static_cast<uint8>(ActiveTags[0]->GetEnumTag());
+		while (AllocationGroups->Num() <= Index)
+		{
+			FAllocationGroup* AddedGroup = Allocator.New<FAllocationGroup>(AllocationGroups->Num());
+			AllocationGroups->Add(AddedGroup);
+		}
+		AllocationGroup = (*AllocationGroups)[Index];
+#endif // else !LLM_ALLOW_NAMES_TAGS
+		AllocationGroup->Initialize(ActiveTags);
+	}
+	return AllocationGroup;
+}
+
+FAllocationGroup* FLLMGlobals::FindAllocationGroupForCompressedAllocInfo(ELLMTag Tag)
+{
+#if LLM_ALLOW_NAMES_TAGS
+	LLMCheck(false);
+	return nullptr;
+#else
+	UE::TSharedLock ScopeLock(AllocationGroupsIndexMutex);
+	uint8 Index = static_cast<uint8>(Tag);
+	// The AllocationGroup should have have been allocated by FindOrAddAllocationGroup during creation of the AllocInfo,
+	// and inserted at an Index equal to its ELLMTag.
+	LLMCheck(AllocationGroups->Num() > Index);
+	return (*AllocationGroups)[Index];
+#endif
+}
+
 } // namespace UE::LLMPrivate
 
 FLLMScope::FLLMScope(ELLMTag TagEnum, bool bInIsStatTag, ELLMTagSet InTagSet, ELLMTracker InTracker, bool bOverride)
@@ -3658,7 +3841,7 @@ bool FLLMScopeDynamic::TryAddTagAndActivate(FName TagName, const ILLMDynamicTagC
 {
 	UE::LLMPrivate::FLLMGlobals& LLMRef = UE::LLMPrivate::FLLMGlobals::GetInner();
 	FName StatFullName = NAME_None;
-#if LLM_ENABLED_STAT_TAGS && STATS
+#if STATS
 	FString StatConstructorName = Constructor.GetStatName();
 	if (!StatConstructorName.IsEmpty())
 	{
@@ -4630,115 +4813,42 @@ bool FTagData::IsUsedAsDisplayParent() const
 	return !(bHasEnumTag && EnumTag == ELLMTag::CustomName);
 }
 
-FLLMTracker::FLowLevelAllocInfo::FLowLevelAllocInfo()
-{
-#if LLM_ENABLED_FULL_TAGS && (LLM_ALLOW_ASSETS_TAGS || LLM_ALLOW_UOBJECTCLASSES_TAGS)
-	FMemory::Memset(Tag, InvalidCompressedTagValue);
-#endif
-}
+#if LLM_ALLOW_NAMES_TAGS
 
-int32 FLLMTracker::FLowLevelAllocInfo::GetIndexFromTagSet(ELLMTagSet InTagSet)
+void FLLMTracker::FLowLevelAllocInfo::SetGroup(FAllocationGroup* AllocationGroup)
 {
-	static_assert(static_cast<int32>(ELLMTagSet::Max) == 4,
-		"You added a new tagset, please update the code bellow");
-
-#if LLM_ALLOW_ASSETS_TAGS
-	switch (InTagSet)
-	{
-	case ELLMTagSet::None: return 0;
-	case ELLMTagSet::Assets: return 1;
-	case ELLMTagSet::AssetClasses: return 2;
-	default:
-		LLMCheck(false);
-		return 0;
-	}
-#elif LLM_ALLOW_UOBJECTCLASSES_TAGS
-	switch (InTagSet)
-	{
-	case ELLMTagSet::None: return 0;
-	case ELLMTagSet::UObjectClasses: return 1;
-	default:
-		LLMCheck(false);
-		return 0;
-	}
-#endif
+	Group = AllocationGroup->GetIndex();
 }
 
-void FLLMTracker::FLowLevelAllocInfo::SetTag(const FTagData* InTag, FLLMGlobals& InLLMRef
-#if LLM_ALLOW_ASSETS_TAGS || LLM_ALLOW_UOBJECTCLASSES_TAGS
-	, ELLMTagSet InTagSet
-#endif
-)
+FAllocationGroup* FLLMTracker::FLowLevelAllocInfo::GetGroup(FLLMGlobals& InLLMRef) const
 {
-#if LLM_ENABLED_FULL_TAGS
-#if !LLM_ALLOW_ASSETS_TAGS && !LLM_ALLOW_UOBJECTCLASSES_TAGS
-	Tag = InTag->GetIndex();
-#else
-	int32 TagSetIndex = GetIndexFromTagSet(InTagSet);
-	Tag[TagSetIndex] = InTag ? InTag->GetIndex() : InvalidCompressedTagValue;
-#endif //!LLM_ALLOW_ASSETS_TAGS
-#else
-	Tag = InTag->GetEnumTag();
-#endif
+	UE::TSharedLock ScopeLock(InLLMRef.AllocationGroupsIndexMutex);
+	return (*InLLMRef.AllocationGroups)[Group];
 }
 
-const FTagData* FLLMTracker::FLowLevelAllocInfo::GetTag(FLLMGlobals& InLLMRef
-#if LLM_ALLOW_ASSETS_TAGS || LLM_ALLOW_UOBJECTCLASSES_TAGS
-	, ELLMTagSet InTagSet
-#endif
-) const
+int32 FLLMTracker::FLowLevelAllocInfo::GetCompressedTag() const
 {
-#if LLM_ENABLED_FULL_TAGS
-	UE::TSharedLock TagDataScopeLock(InLLMRef.TagDataMutex);
-#if !LLM_ALLOW_ASSETS_TAGS && !LLM_ALLOW_UOBJECTCLASSES_TAGS
-	return (*InLLMRef.TagDatas)[Tag];
-#else
-	int32 TagSetIndex = GetIndexFromTagSet(InTagSet);
-	return Tag[TagSetIndex] >= 0 ? (*InLLMRef.TagDatas)[Tag[TagSetIndex]] : nullptr;
-#endif
-#else
-	return InLLMRef.FindTagData(Tag);
-#endif
+	return Group;
 }
 
-#if LLM_ENABLED_FULL_TAGS
-
-int32 FLLMTracker::FLowLevelAllocInfo::GetCompressedTag(ELLMTagSet InTagSet) const
-{
-#if !LLM_ALLOW_ASSETS_TAGS && !LLM_ALLOW_UOBJECTCLASSES_TAGS
-	return (InTagSet == ELLMTagSet::None) ? Tag : InvalidCompressedTagValue;
-#else
-	int32 TagSetIndex = GetIndexFromTagSet(InTagSet);
-	return Tag[TagSetIndex];
-#endif
-}
+#else // LLM_ALLOW_NAMES_TAGS
 
-void FLLMTracker::FLowLevelAllocInfo::SetCompressedTag(int32 InTag, ELLMTagSet InTagSet)
+void FLLMTracker::FLowLevelAllocInfo::SetGroup(FAllocationGroup* AllocationGroup)
 {
-#if !LLM_ALLOW_ASSETS_TAGS && !LLM_ALLOW_UOBJECTCLASSES_TAGS
-	if (InTagSet == ELLMTagSet::None)
-	{
-		Tag = InTag;
-	}
-#else
-	int32 TagSetIndex = GetIndexFromTagSet(InTagSet);
-	Tag[TagSetIndex] = InTag;
-#endif
+	Tag = AllocationGroup->GetActiveTags().GetSystemsTagData()->GetEnumTag();
 }
 
-#else // LLM_ENABLED_FULL_TAGS
-
-ELLMTag FLLMTracker::FLowLevelAllocInfo::GetCompressedTag() const
+FAllocationGroup* FLLMTracker::FLowLevelAllocInfo::GetGroup(FLLMGlobals& InLLMRef) const
 {
-	return Tag;
+	return InLLMRef.FindAllocationGroupForCompressedAllocInfo(Tag);
 }
 
-void FLLMTracker::FLowLevelAllocInfo::SetCompressedTag(ELLMTag InTag)
+int32 FLLMTracker::FLowLevelAllocInfo::GetCompressedTag() const
 {
-	Tag = InTag;
+	return static_cast<int32>(Tag);
 }
 
-#endif // else LLM_ENABLED_FULL_TAGS
+#endif // else !LLM_ALLOW_NAMES_TAGS
 
 FLLMTracker::FLLMTracker(FLLMGlobals& InLLM)
 	: LLMRef(InLLM)
@@ -4860,11 +4970,18 @@ void FLLMTracker::TrackAllocation(const void* Ptr, int64 Size, ELLMTag DefaultTa
 {
 	FLLMThreadState* State = GetOrCreateState();
 	const FTagData* TagData = State->GetTopTag(ELLMTagSet::None);
+	bool bPushed = false;
 	if (!TagData)
 	{
 		TagData = LLMRef.FindOrAddTagData(DefaultTag);
+		State->PushTag(TagData, ELLMTagSet::None);
+		bPushed = true;
+	}
+	TrackAllocationInActiveTags(Ptr, Size, AllocType, State, bTrackInMemPro);
+	if (bPushed)
+	{
+		State->PopTag(ELLMTagSet::None);
 	}
-	TrackAllocation(Ptr, Size, TagData, AllocType, State, bTrackInMemPro);
 }
 
 void FLLMTracker::TrackAllocation(const void* Ptr, int64 Size, FName DefaultTag, ELLMAllocType AllocType,
@@ -4872,14 +4989,22 @@ void FLLMTracker::TrackAllocation(const void* Ptr, int64 Size, FName DefaultTag,
 {
 	FLLMThreadState* State = GetOrCreateState();
 	const FTagData* TagData = State->GetTopTag(ELLMTagSet::None);
+	bool bPushed = false;
 	if (!TagData)
 	{
 		TagData = LLMRef.FindOrAddTagData(DefaultTag, ELLMTagSet::None);
+		State->PushTag(TagData, ELLMTagSet::None);
+		bPushed = true;
+	}
+	TrackAllocationInActiveTags(Ptr, Size, AllocType, State, bTrackInMemPro);
+	if (bPushed)
+	{
+		State->PopTag(ELLMTagSet::None);
 	}
-	TrackAllocation(Ptr, Size, TagData, AllocType, State, bTrackInMemPro);
 }
 
-void FLLMTracker::TrackAllocation(const void* Ptr, int64 Size, const FTagData* ActiveTagData, ELLMAllocType AllocType, FLLMThreadState* State, bool bTrackInMemPro)
+void FLLMTracker::TrackAllocationInActiveTags(const void* Ptr, int64 Size, ELLMAllocType AllocType,
+	FLLMThreadState* State, bool bTrackInMemPro)
 {
 	if (IsPaused(AllocType))
 	{
@@ -4892,36 +5017,9 @@ void FLLMTracker::TrackAllocation(const void* Ptr, int64 Size, const FTagData* A
 	// track the total quickly
 	FPlatformAtomics::InterlockedAdd(&TrackedTotal, Size);
 
-#if !LLM_ENABLED_FULL_TAGS
-	// When full tags are disabled, we instead store the top-level enumtag parent of the allocation's tag.
-	ActiveTagData = ActiveTagData->GetContainingEnumTagData();
-#endif
-
-#if LLM_ALLOW_ASSETS_TAGS
-	const FTagData* AssetTagData = State->GetTopTag(ELLMTagSet::Assets);
-	const FTagData* AssetClassTagData = State->GetTopTag(ELLMTagSet::AssetClasses);
-
-	if (AssetTagData == nullptr && !LLMRef.IsBootstrapping() && LLMRef.IsTagSetActive(ELLMTagSet::Assets))
-	{
-		AssetTagData = LLMRef.FindOrAddTagData(TagName_UntaggedAsset, ELLMTagSet::Assets, false, ETagReferenceSource::FunctionAPI);
-	}
-	if (AssetClassTagData == nullptr && !LLMRef.IsBootstrapping() && LLMRef.IsTagSetActive(ELLMTagSet::AssetClasses))
-	{
-		AssetClassTagData = LLMRef.FindOrAddTagData(TagName_UntaggedAssetClass, ELLMTagSet::AssetClasses, false, ETagReferenceSource::FunctionAPI);
-	}
-#else
-	const FTagData* AssetTagData = nullptr;
-	const FTagData* AssetClassTagData = nullptr;
-#endif // LLM_ALLOW_ASSETS_TAGS
-
-#if LLM_ALLOW_UOBJECTCLASSES_TAGS
-	const FTagData* UObjectClassesTagData = State->GetTopTag(ELLMTagSet::UObjectClasses);
-#else
-	const FTagData* UObjectClassesTagData = nullptr;
-#endif
-
-	// track on the thread state
-	State->TrackAllocation(Ptr, Size, Tracker, AllocType, ActiveTagData, AssetTagData, AssetClassTagData, UObjectClassesTagData, bTrackInMemPro);
+	// Get the allocationgroup for the current threadstate and increment the threadstate's accumulated size for it.
+	FAllocationGroup* AllocationGroup = nullptr;
+	State->TrackAllocation(Ptr, Size, LLMRef, Tracker, AllocType, bTrackInMemPro, AllocationGroup);
 
 #if !UE_ONLY_USE_PLATFORM_TRACKER
 	// tracking a nullptr with a Size is allowed, but we don't need to remember it, since we can't free it ever.
@@ -4929,18 +5027,7 @@ void FLLMTracker::TrackAllocation(const void* Ptr, int64 Size, const FTagData* A
 	{
 		// remember the size and tag info
 		FLLMTracker::FLowLevelAllocInfo AllocInfo;
-		AllocInfo.SetTag(ActiveTagData, LLMRef);
-#if LLM_ALLOW_ASSETS_TAGS
-		AllocInfo.SetTag(AssetTagData, LLMRef, ELLMTagSet::Assets);
-		AllocInfo.SetTag(AssetClassTagData, LLMRef, ELLMTagSet::AssetClasses);
-#endif
-#if LLM_ALLOW_UOBJECTCLASSES_TAGS
-		// If there is no scoped tag for UObjectClasses, don't set it
-		if (UObjectClassesTagData != nullptr)
-		{
-			AllocInfo.SetTag(UObjectClassesTagData, LLMRef, ELLMTagSet::UObjectClasses);
-		}
-#endif
+		AllocInfo.SetGroup(AllocationGroup);
 		LLMCheck(Size <= 0x0000'ffff'ffff'ffff);
 		uint32 SizeLow = uint32(Size);
 		uint16 SizeHigh = uint16(Size >> 32ull);
@@ -4978,33 +5065,14 @@ void FLLMTracker::TrackFree(const void* Ptr, ELLMAllocType AllocType, bool bTrac
 	FPlatformAtomics::InterlockedAdd(&TrackedTotal, -Size);
 
 	FLLMThreadState* State = GetOrCreateState();
-	const FTagData* TagData = AllocInfo.GetTag(LLMRef);
-#if LLM_ALLOW_ASSETS_TAGS
-	const FTagData* AssetTagData = AllocInfo.GetTag(LLMRef, ELLMTagSet::Assets);
-	const FTagData* AssetClassTagData = AllocInfo.GetTag(LLMRef, ELLMTagSet::AssetClasses);
-#else
-	const FTagData* AssetTagData = nullptr;
-	const FTagData* AssetClassTagData = nullptr;
-#endif // LLM_ALLOW_ASSETS_TAGS
-
-#if LLM_ALLOW_UOBJECTCLASSES_TAGS
-	const FTagData* UObjectClassesTagData = AllocInfo.GetTag(LLMRef, ELLMTagSet::UObjectClasses);
-#else
-	const FTagData* UObjectClassesTagData = nullptr;
-#endif // LLM_ALLOW_UOBJECTCLASSES_TAGS
-
+	FAllocationGroup* AllocationGroup = AllocInfo.GetGroup(LLMRef);
 #else
-
 	FLLMThreadState* State = GetOrCreateState();
 	int64 Size = 0;
-	const FTagData* TagData = nullptr;
-	const FTagData* AssetTagData = nullptr;
-	const FTagData* AssetClassTagData = nullptr;
-	const FTagData* UObjectClassesTagData = nullptr;
-
+	FAllocationGroup* AllocationGroup = nullptr;
 #endif // !UE_ONLY_USE_PLATFORM_TRACKER
 
-	State->TrackFree(Ptr, Size, Tracker, AllocType, TagData, AssetTagData, AssetClassTagData, UObjectClassesTagData, bTrackInMemPro);
+	State->TrackFree(Ptr, Size, Tracker, AllocType, AllocationGroup, bTrackInMemPro);
 }
 
 void FLLMTracker::OnAllocMoved(const void* Dest, const void* Source, ELLMAllocType AllocType)
@@ -5031,14 +5099,15 @@ void FLLMTracker::OnAllocMoved(const void* Dest, const void* Source, ELLMAllocTy
 	int64 SizeHigh = Values.Key.GetExtraData();
 	int64 Size = (SizeHigh << 32ull) | SizeLow;
 	const FLLMTracker::FLowLevelAllocInfo& AllocInfo = Values.Value2;
-	const FTagData* TagData = AllocInfo.GetTag(LLMRef);
+	FAllocationGroup* AllocationGroup = AllocInfo.GetGroup(LLMRef);
 #else
 	// TODO: This loses the memory size of an allocation moved within a larger allocation.
 	int64 Size = 0;
 	const FTagData* TagData = nullptr;
+	FAllocationGroup* AllocationGroup = nullptr;
 #endif
 	FLLMThreadState* State = GetOrCreateState();
-	State->TrackMoved(Dest, Source, Size, Tracker, TagData);
+	State->TrackMoved(Dest, Source, Size, Tracker, AllocationGroup);
 }
 
 void FLLMTracker::TrackMemoryOfActiveTag(int64 Amount, FName DefaultTag, ELLMAllocType AllocType)
@@ -5074,13 +5143,13 @@ void FLLMTracker::TrackMemoryOfActiveTag(int64 Amount, const FTagData* TagData,
 	// track the total quickly
 	FPlatformAtomics::InterlockedAdd(&TrackedTotal, Amount);
 
-#if !LLM_ENABLED_FULL_TAGS
+#if !LLM_ALLOW_NAMES_TAGS
 	// When full tags are disabled, we instead store the top-level enumtag parent of the tag used by each allocation
 	TagData = TagData->GetContainingEnumTagData();
 #endif
 
 	// track on the thread state
-	State->TrackMemory(Amount, Tracker, AllocType, TagData);
+	State->TrackMemory(Amount, LLMRef, Tracker, AllocType, TagData);
 }
 
 void FLLMTracker::TrackMemory(ELLMTag Tag, int64 Amount, ELLMAllocType AllocType)
@@ -5097,7 +5166,7 @@ void FLLMTracker::TrackMemory(const FTagData* TagData, int64 Amount, ELLMAllocTy
 {
 	FLLMThreadState* State = GetOrCreateState();
 	FPlatformAtomics::InterlockedAdd(&TrackedTotal, Amount);
-	State->TrackMemory(Amount, Tracker, AllocType, TagData);
+	State->TrackMemory(Amount, LLMRef, Tracker, AllocType, TagData);
 }
 
 void FLLMTracker::PauseAndTrackMemory(FName TagName, ELLMTagSet TagSet, bool bInIsStatTag, int64 Amount, ELLMAllocType AllocType)
@@ -5117,7 +5186,7 @@ void FLLMTracker::PauseAndTrackMemory(const FTagData* TagData, int64 Amount, ELL
 {
 	FLLMThreadState* State = GetOrCreateState();
 	FPlatformAtomics::InterlockedAdd(&TrackedTotal, Amount);
-	State->TrackMemory(Amount, Tracker, AllocType, TagData);
+	State->TrackMemory(Amount, LLMRef, Tracker, AllocType, TagData);
 	LLMCheck(State->PausedCounter[static_cast<int32>(AllocType)] < 127);
 	State->PausedCounter[static_cast<int32>(AllocType)]++;
 }
@@ -5187,11 +5256,7 @@ void FLLMTracker::SetTotalTags(const FTagData* InOverrideUntaggedTagData,
 void FLLMTracker::CaptureTagSnapshot()
 {
 	UpdateThreads();
-	for (FLLMThreadState* ThreadState : ThreadStates)
-	{
-		ThreadState->PropagateChildSizesToParents();
-		ThreadState->FetchAndClearTagSizes(TagSizes, AllocTypeAmounts, false);
-	}
+	FetchAndClearTagSizes(false /* bUpdatePruning */);
 
 	for (TPair<const FTagData*, FTrackerTagSizeData>& It : TagSizes)
 	{
@@ -5216,26 +5281,9 @@ void FLLMTracker::ClearTagSnapshot()
 void FLLMTracker::Update()
 {
 	UpdateThreads();
-	double CurrentTime = FPlatformTime::Seconds();
-	constexpr double UpdateTrimPeriod = 10.0;
-	bool bTrimAllocations = CurrentTime - LastTrimTime > UpdateTrimPeriod;
-	if (bTrimAllocations)
-	{
-		LastTrimTime = CurrentTime;
-		{
-#if !UE_ONLY_USE_PLATFORM_TRACKER
-			AllocationMap.Trim();
-#endif
-		}
-	}
 
 	// Add the values from each thread to the central repository.
-	bool bTrimThreads = bTrimAllocations;
-	for (FLLMThreadState* ThreadState : ThreadStates)
-	{
-		ThreadState->PropagateChildSizesToParents();
-		ThreadState->FetchAndClearTagSizes(TagSizes, AllocTypeAmounts, bTrimThreads);
-	}
+	FetchAndClearTagSizes(true /* bUpdatePruning */);
 
 	// Update peak sizes and external sizes in the central repository.
 	for (TPair<const FTagData*, FTrackerTagSizeData>& It : TagSizes)
@@ -5262,60 +5310,374 @@ void FLLMTracker::Update()
 	}
 }
 
-void FLLMTracker::UpdateThreads()
-{
-	// Consume pending thread states. We must be careful to do all allocations outside of the
-	// PendingThreadStatesGuard guard as that can lead to a deadlock due to contention with
-	// PendingThreadStatesGuard & Locks inside the underlying allocator (i.e. MallocBinned2 -> Mutex).
+void FLLMTracker::FetchAndClearTagSizes(bool bUpdatePruning)
+{
+	EPruningLevel PruningLevel = EPruningLevel::None;
+	if (bUpdatePruning)
+	{
+		PruningLevel = EPruningLevel::Default;
+		double CurrentTime = FPlatformTime::Seconds();
+		constexpr double UpdateTrimPeriod = 10.0;
+		bool bTrimAllocations = CurrentTime - LastTrimTime > UpdateTrimPeriod;
+		if (bTrimAllocations)
+		{ 
+			PruningLevel = EPruningLevel::Max;
+			LastTrimTime = CurrentTime;
+		}
+	}
+
+	// Note that we are NOT halting all threads to take a snapshot; we want to avoid doing that because it is bad for
+	// performance to make all threads wait while we traverse possibly-large containers on all threads. But this means
+	// we can not measure an accurate aggregation of the ActiveAllocs, because
+	// 
+	// UpdateThread: ReadThread 1: 0 Allocs
+	// UpdateThread: ReadThreads [2, K-1]: 0 Allocs
+	// Thread1: Make Allocation -> ActiveAllocs++
+	// ThreadK: Delete pointer that was allocated from Thread1 -> ActiveAllocs--
+	// UpdateThread: ReadThread K, -1 Allocs
+	// Our counter is now -1 allocations, when it should be 0.
+	// 
+	// We therefore need to tolerate ActiveAllocs that can go negative.
+	// We can afford to tolerate that, however, because we also have each thread recording its own value of
+	// GetUpdatesSinceLastAllocationChange, and clearing that to 0 everytime it sees an allocation or free. If we see
+	// two updates in a row with no AllocationChange on all threads, then we know that we have an accurate measure of
+	// the ActiveAllocs. And we do not take any control flow action (most notably, we do not prune an AllocationGroup)
+	// until both its ThreadRefCount and its ActiveAllocsCount reach 0. And ThreadRefCount can only reach zero when all
+	// threads Release, which they only do after PruneUpdatesForTrackingData occurrences of the Update
+	// function with no AllocationChange, and we require PruneUpdatesForTrackingData >= 2.
+	FMapActiveTagsToAllocationGroupTrackingData LocalTrackings;
+	LocalTrackings.Reserve(1000);
+	for (FLLMThreadState* ThreadState : ThreadStates)
 	{
-		PendingThreadStatesGuard.Lock();
-		const int32 NumPendingThreadStatesToConsume = PendingThreadStates.Num();
-		if (NumPendingThreadStatesToConsume > 0)
+		ThreadState->FetchAndClearTagSizes(LocalTrackings, AllocTypeAmounts, PruningLevel);
+	}
+
+	// System tags (ELLMTagSet::None) are hierarchical, with the size of parent tags being increased by the sizes of any child
+	// tags under them. Copy sizes for AllocationGroups with child ELLMTagSet::None tags into the AllocationGroup for the
+	// parent ELLMTagSet::None tags and the unchanged tags for other TagSets. We do not support a hierarchy for other
+	// TagSets because we haven't needed it so far, and solving the propagation problem for multiple parent hierarchies
+	// would be difficult due to the need to handle diamond-shaped inheritance.
+
+	// Add a tracking data for any parent not already in the list. First identify the starting set to add - the parents
+	// of tags in the list that are not already in the list.
+	TLLMSet<FActiveTags> ParentsToAdd;
+	for (TPair<FActiveTags, FAllocationGroupTrackingData>& LocalTrackingPair : LocalTrackings)
+	{
+		const FActiveTags& ActiveTags = LocalTrackingPair.Key;
+		FAllocationGroupTrackingData& LocalTracking = LocalTrackingPair.Value;
+		// We rely on all of the elements in LocalTracking having a group when adding Sizes below.
+		LLMCheck(LocalTracking.IsInitialized() && LocalTracking.GetGroup() != nullptr);
+
+		const FTagData* ParentSystemsTag = ActiveTags.GetSystemsTagData()->GetParent();
+		if (!ParentSystemsTag)
+		{
+			continue;
+		}
+		FActiveTags ParentTags = ActiveTags;
+		ParentTags.SetSystemsTagData(ParentSystemsTag);
+		if (LocalTrackings.Contains(ParentTags))
 		{
-			PendingThreadStatesGuard.Unlock();
-			ThreadStates.Reserve(ThreadStates.Num() + NumPendingThreadStatesToConsume);
-			PendingThreadStatesGuard.Lock();
+			continue;
+		}
 
-			for (int32 i = 0; i < NumPendingThreadStatesToConsume; ++i)
-			{
-				ThreadStates.Add(PendingThreadStates.Pop(EAllowShrinking::No));
-			}
+		if (ParentsToAdd.IsEmpty())
+		{
+			ParentsToAdd.Reserve(100);
 		}
-		PendingThreadStatesGuard.Unlock();
+		ParentsToAdd.Add(ParentTags);
 	}
-}
 
-void FLLMTracker::PublishStats(UE::LLM::ESizeParams SizeParams)
-{
-	if (OverrideTrackedTotalTagData)
+	// Add a trackingdata for all parents and ancestors in our ParentsToAddList that are not already added.
+	// First handle adding parents that already exist in LLMRef.ActiveTagsToAllocationGroup, which we can do under a read lock.
+	TLLMSet<FActiveTags> ParentsToAddUnderWriteLock;
 	{
-		SetMemoryStatByFName(OverrideTrackedTotalTagData->GetStatName(), TrackedTotal);
-		SetMemoryStatByFName(OverrideTrackedTotalTagData->GetSummaryStatName(), TrackedTotal);
+		UE::TSharedLock ReadScopeLock(LLMRef.AllocationGroupsMutex);
+		for (const FActiveTags& DirectParentTags : ParentsToAdd)
+		{
+			FAllocationGroupTrackingData* AncestorTracking = &LocalTrackings.FindOrAdd(DirectParentTags);
+			if (AncestorTracking->IsInitialized())
+			{
+				continue;
+			}
+
+			FActiveTags AncestorTags = DirectParentTags;
+			for (;;)
+			{
+				FAllocationGroup** AncestorGroupPtr = LLMRef.ActiveTagsToAllocationGroup->Find(AncestorTags);
+				if (!AncestorGroupPtr)
+				{
+					if (ParentsToAddUnderWriteLock.IsEmpty())
+					{
+						ParentsToAddUnderWriteLock.Reserve(100);
+					}
+					ParentsToAddUnderWriteLock.Add(AncestorTags);
+					break;
+				}
+				FAllocationGroup* AncestorGroup = *AncestorGroupPtr;
+
+				AncestorTracking->ConditionalInitialize(AncestorGroup);
+				const FTagData* AncestorSystemsTag = AncestorTags.GetSystemsTagData()->GetParent();
+				if (!AncestorSystemsTag)
+				{
+					break;
+				}
+
+				AncestorTags.SetSystemsTagData(AncestorSystemsTag);
+				AncestorTracking = &LocalTrackings.FindOrAdd(AncestorTags);
+				if (AncestorTracking->IsInitialized())
+				{
+					break;
+				}
+			}
+		}
 	}
 
-	if (OverrideUntaggedTagData)
+	// Add a trackingdata for all parents and ancestors in our ParentsToAddList that do not already exist in
+	// LLMRef.ActiveTagsToAllocationGroup; this requires a WriteLock to create the AllocationGroup.
 	{
-		const FTagData* TagData = LLMRef.FindTagData(TagName_Untagged, ELLMTagSet::None);
-		const FTrackerTagSizeData* AllocationData = TagData ? TagSizes.Find(TagData) : nullptr;
-		SetMemoryStatByFName(OverrideUntaggedTagData->GetStatName(),
-			AllocationData ? AllocationData->GetSize(SizeParams) : 0);
-		SetMemoryStatByFName(OverrideUntaggedTagData->GetSummaryStatName(),
-			AllocationData ? AllocationData->GetSize(SizeParams) : 0);
+		UE::TUniqueLock WriteScopeLock(LLMRef.AllocationGroupsMutex);
+		for (const FActiveTags& DirectParentTags : ParentsToAddUnderWriteLock)
+		{
+			FAllocationGroupTrackingData* AncestorTracking = &LocalTrackings.FindOrAdd(DirectParentTags);
+			if (AncestorTracking->IsInitialized())
+			{
+				continue;
+			}
+
+			FActiveTags AncestorTags = DirectParentTags;
+			for (;;)
+			{
+				FAllocationGroup* AncestorGroup = LLMRef.FindOrAddAllocationGroup(AncestorTags);
+
+				AncestorTracking->ConditionalInitialize(AncestorGroup);
+				const FTagData* AncestorSystemsTag = AncestorTags.GetSystemsTagData()->GetParent();
+				if (!AncestorSystemsTag)
+				{
+					break;
+				}
+
+				AncestorTags.SetSystemsTagData(AncestorSystemsTag);
+				AncestorTracking = &LocalTrackings.FindOrAdd(AncestorTags);
+				if (AncestorTracking->IsInitialized())
+				{
+					break;
+				}
+			}
+		}
 	}
 
-	for (const TPair<const FTagData*, FTrackerTagSizeData>& It : TagSizes)
+	// Sort the TrackingDatas from leaf to root index of their ELLMTag::None tags, so that we can accumulate
+	// Allocation deltas by adding it to one parent at a time. FTagData->GetIndex() is sorted from root to
+	// leaf index, so we can use that for the leaf to root index.
+	// Note that when in this function we are outside of all LLM mutexes except for the UpdateLock, so it is
+	// okay if the sort function allocates memory.
+	LocalTrackings.KeySort([](const FActiveTags& A, const FActiveTags& B)
+		{
+			return A.GetSystemsTagData()->GetIndex() > B.GetSystemsTagData()->GetIndex();
+		});
+
+	// Iterate over every TrackingData we have in LocalTrackings, and add its Size and Allocs to its own group and to
+	// the TrackingData for its parent group.
+	for (TPair<FActiveTags, FAllocationGroupTrackingData>& LocalTrackingPair : LocalTrackings)
 	{
-		const FTagData* TagData = It.Key;
-		if (!TagData->IsStatsReportable())
+		const FActiveTags& ActiveTags = LocalTrackingPair.Key;
+		FAllocationGroupTrackingData& LocalTracking = LocalTrackingPair.Value;
+		// All TrackingDatas we received from ThreadState->FetchAndClearTagSizes were guaranteed initialized, and we
+		// initialized each of the parents that we added in one of the two loops above, so we can assert all elements
+		// in LocalTracking are intialized by now.
+		LLMCheck(LocalTracking.IsInitialized());
+		LLMCheck(LocalTracking.GetGroup() != nullptr);
+		FetchAndClearTagSizes_UpdateGroupAllocs(*LocalTracking.GetGroup(), LocalTracking, PruningLevel);
+
+		if (LocalTracking.GetSize() == 0 && LocalTracking.GetActiveAllocs() == 0)
 		{
 			continue;
 		}
-		if (OverrideUntaggedTagData && TagData->GetName() == TagName_Untagged)
+
+		const FTagData* ParentSystemsTag = ActiveTags.GetSystemsTagData()->GetParent();
+		if (!ParentSystemsTag)
 		{
-			// Handled separately by OverrideUntaggedTagData.
 			continue;
 		}
-		const FTrackerTagSizeData& AllocationData = It.Value;
+
+		FActiveTags ParentTags = ActiveTags;
+		ParentTags.SetSystemsTagData(ParentSystemsTag);
+		FAllocationGroupTrackingData* ParentTracking = LocalTrackings.Find(ParentTags);
+		// We added the Tracking data for ParentTags in one of the two loops above, so assert that we have it.
+		LLMCheck(ParentTracking && ParentTracking->IsInitialized());
+
+		// Accumulate the Size and ActiveAllocs up to the parent
+		// Note: the ActiveAllocs count for a parent tag might incorrectly be zero, without the parent
+		// tag being referenced from any thread. So without intervention here we might decide to prune it because
+		// it is not referenced. This will not be a behavior problem, because we will just recreate it (with
+		// the correct aggregate value of 0 ActiveAllocs and 0 ThreadRefCount) when it is next referenced,
+		// but that will be a performance problem. To prevent that problem, ResetUpdatesSinceLastReference whenever
+		// a tag is referenced as a parent during in update.
+		ParentTracking->GetGroup()->ResetUpdatesSinceLastReference();
+		ParentTracking->TrackAllocOrFree(LocalTracking.GetSize(), LocalTracking.GetActiveAllocs());
+	}
+
+	if (PruningLevel > EPruningLevel::None)
+	{
+#if !UE_ONLY_USE_PLATFORM_TRACKER
+		if (PruningLevel >= EPruningLevel::Max)
+		{
+			AllocationMap.Trim();
+		}
+#endif
+
+		if (this->Tracker == ELLMTracker::Default)
+		{
+			// Prune the AllocationGroups in the unreferencedlist
+			TArray<FAllocationGroup*, FDefaultLLMAllocator> GroupsToFree;
+			for (auto GroupIter = LLMRef.AllocationGroupsUnreferencedSet->CreateIterator(); GroupIter; ++GroupIter)
+			{
+				FAllocationGroup* Group = *GroupIter;
+				if (Group->IsReferenced())
+				{
+					GroupIter.RemoveCurrent();
+					continue;
+				}
+				Group->IncrementUpdatesSinceLastReference();
+				if (Group->GetUpdatesSinceLastReference() >= PruneUpdatesForGroup)
+				{
+					GroupsToFree.Add(Group);
+					GroupIter.RemoveCurrent();
+					continue;
+				}
+			}
+			if (!GroupsToFree.IsEmpty())
+			{
+				UE::TUniqueLock WriteScopeLock(LLMRef.AllocationGroupsMutex);
+				for (FAllocationGroup* Group : GroupsToFree)
+				{
+					// Check RefCount again while under the WriteScopeLock. It might have changed in between our check
+					// of it outside the lock and now, because a thread can lookup the group from
+					// ActiveTagsToAllocationGroup and increment the refcount. We do not have to handle it changing
+					// after this point, because it only increments RefCount when inside the AllocationGroupsMutex.
+					if (Group->IsReferenced())
+					{
+						// The group no longer needs to be freed and also no longer needs to be in the
+						// AllocationGroupsUnreferencedSet.
+						// Drop it from our consideration and keep it in ActiveTagsToAllocationGroup.
+						continue;
+					}
+
+					// We do not support Resetting AllocationGroups in the !LLM_ALLOW_NAMES_TAGS case, because it's unnecessary since there
+					// are so few of them and because FindOrAddAllocationGroup does not yet implement allocating them from the freelist
+					// while preserving their Index correctly.
+#if LLM_ALLOW_NAMES_TAGS
+					LLMRef.ActiveTagsToAllocationGroup->Remove(Group->GetActiveTags());
+					Group->Reset();
+					LLMRef.AllocationGroupsFreeList->Add(Group);
+#endif // LLM_ALLOW_NAMES_TAGS
+				}
+			}
+
+			if (PruningLevel >= EPruningLevel::Max)
+			{
+				UE::TUniqueLock WriteScopeLock(LLMRef.AllocationGroupsMutex);
+				LLMRef.ActiveTagsToAllocationGroup->Shrink();
+				LLMRef.AllocationGroupsUnreferencedSet->Shrink();
+				LLMRef.AllocationGroupsFreeList->Shrink();
+				{
+					UE::TUniqueLock WriteScopeLockGlobalIndex(LLMRef.AllocationGroupsIndexMutex);
+					LLMRef.AllocationGroups->Shrink();
+				}
+			}
+		}
+	}
+}
+
+void FLLMTracker::FetchAndClearTagSizes_UpdateGroupAllocs(FAllocationGroup& AllocationGroup,
+	const FAllocationGroupTrackingData& TrackingData, EPruningLevel PruningLevel)
+{
+	int64 Size = TrackingData.GetSize();
+	AllocationGroup.AddSizeAndAllocReferences(this->Tracker, Size, TrackingData.GetActiveAllocs(),
+		PruningLevel);
+	if (Size != 0)
+	{
+		for (const FTagData* TagData : AllocationGroup.GetActiveTags())
+		{
+			FTrackerTagSizeData& TrackerInfo = TagSizes.FindOrAdd(TagData);
+			TrackerInfo.Size += Size;
+		}
+	}
+
+	if (PruningLevel > EPruningLevel::None)
+	{
+		if (!AllocationGroup.IsReferenced())
+		{
+			LLMRef.AllocationGroupsUnreferencedSet->Add(&AllocationGroup);
+		}
+	}
+}
+
+void FLLMTracker::UpdateThreads()
+{
+	// Consume pending thread states. We must be careful to do all allocations outside of the
+	// PendingThreadStatesGuard guard as that can lead to a deadlock due to contention with
+	// PendingThreadStatesGuard & Locks inside the underlying allocator (i.e. MallocBinned2 -> Mutex).
+	{
+		PendingThreadStatesGuard.Lock();
+		int32 ReservationSize = PendingThreadStates.Num();
+		if (ReservationSize > 0)
+		{
+			// Exit the lock so we can do the allocation, and then reenter the lock. Keep looping until the
+			// Reservation we allocate outside the lock is equal to the reservation we see we need while when we
+			// reenter the lock.
+			for (;;)
+			{
+				PendingThreadStatesGuard.Unlock();
+				ThreadStates.Reserve(ThreadStates.Num() + ReservationSize);
+				PendingThreadStatesGuard.Lock();
+
+				if (PendingThreadStates.Num() != ReservationSize)
+				{
+					ReservationSize = PendingThreadStates.Num();
+					continue;
+				}
+
+				ThreadStates.Append(PendingThreadStates);
+				PendingThreadStates.Reset();
+				break;
+			}
+		}
+		PendingThreadStatesGuard.Unlock();
+	}
+}
+
+void FLLMTracker::PublishStats(UE::LLM::ESizeParams SizeParams)
+{
+	if (OverrideTrackedTotalTagData)
+	{
+		SetMemoryStatByFName(OverrideTrackedTotalTagData->GetStatName(), TrackedTotal);
+		SetMemoryStatByFName(OverrideTrackedTotalTagData->GetSummaryStatName(), TrackedTotal);
+	}
+
+	if (OverrideUntaggedTagData)
+	{
+		const FTagData* TagData = LLMRef.FindTagData(TagName_Untagged, ELLMTagSet::None);
+		const FTrackerTagSizeData* AllocationData = TagData ? TagSizes.Find(TagData) : nullptr;
+		SetMemoryStatByFName(OverrideUntaggedTagData->GetStatName(),
+			AllocationData ? AllocationData->GetSize(SizeParams) : 0);
+		SetMemoryStatByFName(OverrideUntaggedTagData->GetSummaryStatName(),
+			AllocationData ? AllocationData->GetSize(SizeParams) : 0);
+	}
+
+	for (const TPair<const FTagData*, FTrackerTagSizeData>& It : TagSizes)
+	{
+		const FTagData* TagData = It.Key;
+		if (!TagData->IsStatsReportable())
+		{
+			continue;
+		}
+		if (OverrideUntaggedTagData && TagData->GetName() == TagName_Untagged)
+		{
+			// Handled separately by OverrideUntaggedTagData.
+			continue;
+		}
+		const FTrackerTagSizeData& AllocationData = It.Value;
 		int64 Amount = AllocationData.GetSize(SizeParams);
 
 		SetMemoryStatByFName(TagData->GetStatName(), Amount);
@@ -5334,6 +5696,10 @@ static bool PageFinder(const FForkedPageAllocation& LHS, const uint64 AddressRHS
 
 bool FLLMTracker::DumpForkedAllocationInfo()
 {
+#if UE_ONLY_USE_PLATFORM_TRACKER
+	UE_LOG(LogHAL, Error, TEXT("UE_ONLY_USE_PLATFORM_TRACKER=1 is defined so reading Allocations from LLM is not available."));
+	return false;
+#else
 	// Try to associate the allocations with a page range so we can determine whether it's
 	// in unique or shared memory. Then print out the set for each tag.
 	TArray<FForkedPageAllocation> Pages;
@@ -5372,151 +5738,184 @@ bool FLLMTracker::DumpForkedAllocationInfo()
 	};
 
 
-	// We cannot call GetTag to find the Tag from the Allocation when we are within AllocationMap, because GetTag
-	// locks the TagDataMutex and TagDataMutex has to be entered before entering AllocationMap.LockAll.
-	// So for the TagIdentifier use the CompressedTag rather than calling GetTag.
-#if LLM_ENABLED_FULL_TAGS
-	TMap<int32, FCounts> CountsPerTag;
-#else
-	TMap<ELLMTag, FCounts> CountsPerTag;
-#endif
-	int32 NumTags;
+	TMap<const FTagData*, FCounts, FDefaultSetLLMAllocator> CountsPerTag;
 	{
 		UE::TSharedLock Lock(LLMRef.TagDataMutex);
-		NumTags = LLMRef.TagDatas->Num();
+		CountsPerTag.Reserve(LLMRef.TagDatas->Num());
 	}
 
-	CountsPerTag.Reserve(NumTags);
-#if !UE_ONLY_USE_PLATFORM_TRACKER
-	AllocationMap.LockAll();
-	for (const FLLMAllocMap::FTuple& Tuple : AllocationMap)
+	enum class EErrorType
 	{
-		void* Ptr = Tuple.Key.GetPointer();
-
-		int FirstPageEqualOrAfterPtr = Algo::LowerBound(Pages, (uint64)Ptr, PageFinder);
-		int ContainingAllocationIndex;
-		if (FirstPageEqualOrAfterPtr < Pages.Num() && Pages[FirstPageEqualOrAfterPtr].PageStart == (uint64)Ptr)
-		{
-			ContainingAllocationIndex = FirstPageEqualOrAfterPtr;
-		}
-		else
-		{
-			ContainingAllocationIndex = FirstPageEqualOrAfterPtr - 1;
-		}
-		if (ContainingAllocationIndex < 0 || Pages[ContainingAllocationIndex].PageEnd <= (uint64)Ptr)
+		NotInPagesList,
+		InvalidAllocationGroup,
+		InvalidSystemsTag,
+		InvalidSize,
+		ExtendedBeyondPages,
+	};
+	struct FErrorData
+	{
+		uint64 Ptr;
+		EErrorType ErrorType;
+	};
+	TArray<FErrorData, FDefaultLLMAllocator> ErrorPtrs;
+	EnumerateAllocations([&Pages, &ErrorPtrs, &CountsPerTag](void* Ptr, int64 InSize, FAllocationGroup* AllocationGroup)
 		{
-			UE_LOG(LogHAL, Error, TEXT("Can't find allocation 0x%llx in the pages list!"), (uint64)Ptr);
-			continue;
-		}
+			int FirstPageEqualOrAfterPtr = Algo::LowerBound(Pages, (uint64)Ptr, PageFinder);
+			int ContainingAllocationIndex;
+			if (FirstPageEqualOrAfterPtr < Pages.Num() && Pages[FirstPageEqualOrAfterPtr].PageStart == (uint64)Ptr)
+			{
+				ContainingAllocationIndex = FirstPageEqualOrAfterPtr;
+			}
+			else
+			{
+				ContainingAllocationIndex = FirstPageEqualOrAfterPtr - 1;
+			}
+			if (ContainingAllocationIndex < 0 || Pages[ContainingAllocationIndex].PageEnd <= (uint64)Ptr)
+			{
+				ErrorPtrs.Add({ (uint64)Ptr, EErrorType::NotInPagesList });
+				return;
+			}
+			if (AllocationGroup == nullptr)
+			{
+				ErrorPtrs.Add({ (uint64)Ptr, EErrorType::InvalidAllocationGroup });
+				return;
+			}
+			const FTagData* SystemsTag = AllocationGroup->GetActiveTags().GetSystemsTagData();
+			if (SystemsTag == nullptr)
+			{
+				ErrorPtrs.Add({ (uint64)Ptr, EErrorType::InvalidSystemsTag });
+				return;
+			}
+			if (InSize < 0)
+			{
+				ErrorPtrs.Add({ (uint64)Ptr, EErrorType::InvalidSize });
+				return;
+			}
+			// We unfortunately don't know exactly which of the details apply because we don't know which pages
+			// in the allocation section are shared/unique etc, so we just keep some generalities.
+			uint64 Size = static_cast<uint64>(InSize);
+
+			// Number of page allocations for the ue allocation
+			uint32 PageAllocationCount = 0;
+
+			// Track the classification for each page allocation so we can try to classify the ue allocation
+			uint32 SplitCount = 0;
+			uint32 SharedCount = 0;
+			uint32 PrivateCount = 0;
+			uint32 UnreferencedCount = 0;
+
+			// Walk the page allocations until we get them all if we aren't all in the same one.
+			uint64 Remaining = Size;
+			while (Remaining)
+			{
+				if (ContainingAllocationIndex >= Pages.Num())
+				{
+					ErrorPtrs.Add({ (uint64)Ptr, EErrorType::ExtendedBeyondPages });
+					break;
+				}
 
-		// We unfortunately don't know exactly which of the details apply because we don't know which pages
-		// in the allocation section are shared/unique etc, so we just keep some generalities.
-		uint64 SizeHigh = Tuple.Key.GetExtraData() << 32;
-		uint64 SizeLow = Tuple.Value1;
-		uint64 Size = SizeHigh | SizeLow;
+				PageAllocationCount++;
+				FForkedPageAllocation& Page = Pages[ContainingAllocationIndex];
 
-		// Number of page allocations for the ue allocation
-		uint32 PageAllocationCount = 0;
+				uint64 PageKiB = (Page.PageEnd - Page.PageStart) / 1024;
 
-		// Track the classification for each page allocation so we can try to classify the ue allocation
-		uint32 SplitCount = 0;
-		uint32 SharedCount = 0;
-		uint32 PrivateCount = 0;
-		uint32 UnreferencedCount = 0;
+				uint64 OffsetInPage = (uint64)Ptr - Page.PageStart;
+				uint64 RemainingInPage = Page.PageEnd - (uint64)Ptr;
 
-		// Walk the page allocations until we get them all if we aren't all in the same one.
-		uint64 Remaining = Size;
-		while (Remaining)
-		{
-			if (ContainingAllocationIndex >= Pages.Num())
-			{
-				UE_LOG(LogHAL, Error, TEXT("Allocation 0x%llu extended beyond the pages!"), (uint64)Ptr);
-				break;
-			}
+				uint64 AmountInThisPage = Remaining;
+				if (AmountInThisPage > RemainingInPage)
+				{
+					AmountInThisPage = RemainingInPage;
+				}
 
-			PageAllocationCount++;
-			FForkedPageAllocation& Page = Pages[ContainingAllocationIndex];
+				uint64 SharedKiB = Page.SharedCleanKiB + Page.SharedDirtyKiB;
+				uint64 PrivateKiB = Page.PrivateCleanKiB + Page.PrivateDirtyKiB;
+				uint64 UnreferencedKiB = PageKiB - SharedKiB - PrivateKiB;
 
-			uint64 PageKiB = (Page.PageEnd - Page.PageStart) / 1024;
+				// We can't classify the page if we have more than one of any.
+				uint64 TypeCount = !!SharedKiB + !!PrivateKiB + !!UnreferencedKiB;
+				bool bIsSplit = TypeCount > 1;
 
-			uint64 OffsetInPage = (uint64)Ptr - Page.PageStart;
-			uint64 RemainingInPage = Page.PageEnd - (uint64)Ptr;
+				if (bIsSplit)
+				{
+					// The page is split and we have no way to know which is ours.
+					SplitCount++;
+				}
+				else if (SharedKiB)
+				{
+					SharedCount++;
+				}
+				else if (PrivateKiB)
+				{
+					PrivateCount++;
+				}
+				else
+				{
+					UnreferencedCount++;
+				}
 
-			uint64 AmountInThisPage = Remaining;
-			if (AmountInThisPage > RemainingInPage)
-			{
-				AmountInThisPage = RemainingInPage;
+				Remaining -= AmountInThisPage;
+				ContainingAllocationIndex++;
 			}
 
-			uint64 SharedKiB = Page.SharedCleanKiB + Page.SharedDirtyKiB;
-			uint64 PrivateKiB = Page.PrivateCleanKiB + Page.PrivateDirtyKiB;
-			uint64 UnreferencedKiB = PageKiB - SharedKiB - PrivateKiB;
-
-			// We can't classify the page if we have more than one of any.
-			uint64 TypeCount = !!SharedKiB + !!PrivateKiB + !!UnreferencedKiB;
-			bool bIsSplit = TypeCount > 1;
+			if (Remaining)
+			{
+				// We error'd out, ignore this allocation.
+				return;
+			}
 
-			if (bIsSplit)
+			// Classify the allocation - is it entirely private, entirely shared, or what.
+			uint64 ClassCount = !!SharedCount + !!PrivateCount + !!UnreferencedCount;
+			EClassification AllocClass = Class_Split;
+			if (ClassCount > 1)
 			{
-				// The page is split and we have no way to know which is ours.
-				SplitCount++;
+				AllocClass = Class_Split;
 			}
-			else if (SharedKiB)
+			else if (SharedCount)
 			{
-				SharedCount++;
+				AllocClass = Class_Shared;
 			}
-			else if (PrivateKiB)
+			else if (PrivateCount)
 			{
-				PrivateCount++;
+				AllocClass = Class_Private;
 			}
 			else
 			{
-				UnreferencedCount++;
+				AllocClass = Class_Unreferenced;
 			}
 
-			Remaining -= AmountInThisPage;
-			ContainingAllocationIndex++;
-		}
-
-		if (Remaining)
-		{
-			// We error'd out, ignore this allocation.
-			continue;
-		}
-
-		// Classify the allocation - is it entirely private, entirely shared, or what.
-		uint64 ClassCount = !!SharedCount + !!PrivateCount + !!UnreferencedCount;
-		EClassification AllocClass = Class_Split;
-		if (ClassCount > 1)
-		{
-			AllocClass = Class_Split;
-		}
-		else if (SharedCount)
-		{
-			AllocClass = Class_Shared;
-		}
-		else if (PrivateCount)
-		{
-			AllocClass = Class_Private;
-		}
-		else
-		{
-			AllocClass = Class_Unreferenced;
-		}
-
-		// Associate the tag and the page stats.
-		FCounts& Counts = CountsPerTag.FindOrAdd(Tuple.Value2.GetCompressedTag());
-		Counts.AllocCount[AllocClass]++;
-		Counts.ByteCount[AllocClass] += Size;
-		Counts.TotalAllocations++;
-		Counts.TotalBytes += Size;
-		if (PageAllocationCount > 1)
+			// Associate the tag and the page stats.
+			FCounts& Counts = CountsPerTag.FindOrAdd(SystemsTag);
+			Counts.AllocCount[AllocClass]++;
+			Counts.ByteCount[AllocClass] += Size;
+			Counts.TotalAllocations++;
+			Counts.TotalBytes += Size;
+			if (PageAllocationCount > 1)
+			{
+				Counts.CrossPageAllocationCount++;
+			}
+		});
+	for (FErrorData& ErrorData : ErrorPtrs)
+	{
+		switch (ErrorData.ErrorType)
 		{
-			Counts.CrossPageAllocationCount++;
+		case EErrorType::NotInPagesList:
+			UE_LOG(LogHAL, Error, TEXT("Can't find allocation 0x%llx in the pages list!"), ErrorData.Ptr);
+			break;
+		case EErrorType::InvalidAllocationGroup:
+			UE_LOG(LogHAL, Error, TEXT("Allocation 0x%llx has a null AllocationGroup!"), ErrorData.Ptr);
+			break;
+		case EErrorType::InvalidSystemsTag:
+			UE_LOG(LogHAL, Error, TEXT("Allocation 0x%llx has a null SystemsTag!"), ErrorData.Ptr);
+			break;
+		case EErrorType::InvalidSize:
+			UE_LOG(LogHAL, Error, TEXT("Allocation 0x%llx has negative size!"), ErrorData.Ptr);
+			break;
+		case EErrorType::ExtendedBeyondPages:
+			UE_LOG(LogHAL, Error, TEXT("Allocation 0x%llu extended beyond the pages!"), ErrorData.Ptr);
+			break;
 		}
 	}
-	AllocationMap.UnlockAll();
-#endif // !UE_ONLY_USE_PLATFORM_TRACKER
 
 	CountsPerTag.ValueSort([](const FCounts& A, const FCounts& B)
 	{
@@ -5529,19 +5928,10 @@ bool FLLMTracker::DumpForkedAllocationInfo()
 	Lines.Add(TEXT("Tag,SharedKib,PrivateKib,SplitKib,UnrefKib,TotalKib,SharedCount,PrivateCount,SplitCount,UnrefCount,TotalCount,CrossCount"));
 
 
-#if LLM_ENABLED_FULL_TAGS
-	for (const TPair<int32, FCounts>& P : CountsPerTag)
-#else
-	for (const TPair<ELLMTag, FCounts>& P : CountsPerTag)
-#endif
+	for (const TPair<const FTagData*, FCounts>& P : CountsPerTag)
 	{
 		FLowLevelAllocInfo AllocInfoPlaceholder;
-		AllocInfoPlaceholder.SetCompressedTag(P.Key);
-		const FTagData* Tag = AllocInfoPlaceholder.GetTag(LLMRef);
-		if (!Tag)
-		{
-			continue;
-		}
+		const FTagData* Tag = P.Key;
 		Lines.Add(
 			FString::Printf(TEXT("%s,%llu,%llu,%llu,%llu,%llu,%llu,%llu,%llu,%llu,%llu,%llu"),
 				*Tag->GetDisplayName().ToString(),
@@ -5590,6 +5980,7 @@ bool FLLMTracker::DumpForkedAllocationInfo()
 
 	UE_LOG(LogHAL, Log, TEXT("Wrote LLM Private/Shared CSV file: %s"), *UniqueFilename);
 	return true;
+#endif // else !UE_ONLY_USE_PLATFORM_TRACKER
 }
 
 bool FLLMTracker::DumpTags()
@@ -5670,41 +6061,11 @@ void FLLMTracker::PublishCsvProfiler(UE::LLM::ESizeParams SizeParams)
 
 void FLLMTracker::OnTagsResorted(FTagDataArray& OldTagDatas)
 {
-#if LLM_ENABLED_FULL_TAGS && !UE_ONLY_USE_PLATFORM_TRACKER
-	{
-		// Each allocation references the tag by its index, which we have just remapped.
-		// Remap each allocation's tag index to the new index for the tag.
-		AllocationMap.LockAll();
-		for (const FLLMAllocMap::FTuple& Tuple : AllocationMap)
-		{
-			for (int32 TagSetAsInteger = 0; TagSetAsInteger < static_cast<int32>(ELLMTagSet::Max); TagSetAsInteger++)
-			{
-				if (!LLMRef.IsTagSetActive(static_cast<ELLMTagSet>(TagSetAsInteger)))
-				{
-					continue;
-				}
-
-				const int32 OldTagIndex = Tuple.Value2.GetCompressedTag(static_cast<ELLMTagSet>(TagSetAsInteger));
-				if (OldTagIndex >= 0)
-				{
-					Tuple.Value2.SetCompressedTag(OldTagDatas[OldTagIndex]->GetIndex(), static_cast<ELLMTagSet>(TagSetAsInteger));
-				}
-			}
-		}
-		AllocationMap.UnlockAll();
-	}
-#else
-	// Values in AllocationMap are ELLMTags, and don't depend on the Index of the tagdatas.
-#endif
-
-	// Update the uses of Index in the ThreadStates.
-	for (FLLMThreadState* ThreadState : ThreadStates)
-	{
-		ThreadState->OnTagsResorted(OldTagDatas);
-	}
+	// In current code we do not have any persistent sorted list of tags by Index; we create one every update instead.
+	// Keep this hook available in case that ever changes.
 }
 
-void FLLMTracker::LockAllThreadTags(bool bLock)
+void FLLMTracker::LockAllThreadsSharedData(bool bLock)
 {
 	if (bLock)
 	{
@@ -5714,7 +6075,7 @@ void FLLMTracker::LockAllThreadTags(bool bLock)
 
 	for (FLLMThreadState* ThreadState : ThreadStates)
 	{
-		ThreadState->LockTags(bLock);
+		ThreadState->LockSharedData(bLock);
 	}
 
 	if (!bLock)
@@ -5723,6 +6084,28 @@ void FLLMTracker::LockAllThreadTags(bool bLock)
 	}
 }
 
+void FLLMTracker::BeginActiveSetsChange(
+	TArray<TArray<FAllocationGroupTrackingData, FDefaultLLMAllocator>, FDefaultLLMAllocator>& ThreadStateAllocationGroupTrackingDatas)
+{
+	ThreadStateAllocationGroupTrackingDatas.SetNum(ThreadStates.Num());
+	for (int32 Index = 0; Index < ThreadStates.Num(); ++Index)
+	{
+		FLLMThreadState* ThreadState = ThreadStates[Index];
+		ThreadState->BeginActiveSetsChange(ThreadStateAllocationGroupTrackingDatas[Index]);
+	}
+}
+
+void FLLMTracker::EndActiveSetsChange(const FActiveTags& NewDefaults,
+	TArray<TArray<FAllocationGroupTrackingData, FDefaultLLMAllocator>, FDefaultLLMAllocator>& ThreadStateAllocationGroupTrackingDatas)
+{
+	ThreadStateAllocationGroupTrackingDatas.SetNum(ThreadStates.Num());
+	for (int32 Index = 0; Index < ThreadStates.Num(); ++Index)
+	{
+		FLLMThreadState* ThreadState = ThreadStates[Index];
+		ThreadState->EndActiveSetsChange(NewDefaults, ThreadStateAllocationGroupTrackingDatas[Index]);
+	}
+}
+
 const FTagData* FLLMTracker::GetActiveTagData(ELLMTagSet TagSet)
 {
 	FLLMThreadState* State = GetOrCreateState();
@@ -5752,91 +6135,147 @@ void FLLMTracker::GetTagsNamesWithAmount(TMap<FName, uint64>& OutTagsNamesWithAm
 	}
 }
 
-void FLLMTracker::GetTagsNamesWithAmountFiltered(TMap<FName, uint64>& OutTagsNamesWithAmount, ELLMTagSet TagSet /* = ELLMTagSet::None */, TArray<FLLMTagSetAllocationFilter>& Filters)
+int32 FLLMGlobals::GetIndexInFActiveTags(ELLMTagSet TagSet)
 {
-	// We cannot call GetTag to find the Tag from the Allocation when we are within AllocationMap, because GetTag
-	// locks the TagDataMutex and TagDataMutex has to be entered before entering AllocationMap.LockAll.
-	// So for the TagIdentifier use the CompressedTag rather than calling GetTag.
-	struct FLocalTagData
+	int32 Index = 0;
+	for (int32 TagSetIndex = 0; TagSetIndex < static_cast<uint8>(ELLMTagSet::Max); ++TagSetIndex)
 	{
-		FName Name;
-		uint64 Size = 0;
-	};
-	TArray<FLocalTagData> TagDataCached;
-	{
-		UE::TSharedLock Lock(LLMRef.TagDataMutex);
-		const int32 NumTags = LLMRef.TagDatas->Num();
-		TagDataCached.Init(FLocalTagData{}, NumTags);
-		for (FTagData* TagData : (*LLMRef.TagDatas))
+		ELLMTagSet IterTagSet = static_cast<ELLMTagSet>(TagSetIndex);
+		if (IterTagSet == TagSet)
 		{
-			FLLMTracker::FLowLevelAllocInfo AllocInfo;
-			AllocInfo.SetTag(TagData, LLMRef);
-			int32 CompressedTagIndex = static_cast<int32>(AllocInfo.GetCompressedTag());
-			if(ensure(TagDataCached.IsValidIndex(CompressedTagIndex)))
+			if (!ActiveSets[TagSetIndex])
 			{
-				FLocalTagData& Data = TagDataCached[CompressedTagIndex];
-				if (Data.Name == NAME_None)
-				{
-					Data.Name = TagData->GetName();
-				}
+				return -1;
 			}
+			else
+			{
+				return Index;
+			}
+		}
+		if (ActiveSets[TagSetIndex])
+		{
+			++Index;
 		}
 	}
+	return -1;
+}
+
+void FLLMTracker::EnumerateAllocations(TFunctionRef<void(void* Ptr, int64 Size, FAllocationGroup* Group)> Callback)
+{
 #if !UE_ONLY_USE_PLATFORM_TRACKER
+	// Callers must call outside of any LLM locks, and not call LLM functions during callback, because this function
+	// enters the AllocationsLock, which must be not be simultaneously entered with other LLM locks in incorrect lock
+	// ordering.
+
+	// We cannot call GetGroup or read AllocationGroups to find the AllocationGroup used by each AllocInfo when we are
+	// within AllocationMap.LockAll(), because reading AllocationGroups or calling GetGroup locks AllocationGroupsIndexMutex
+	// and the mutex enter order between AllocationGroupsIndexMutex and the AllocationMap's lock is that
+	// AllocationGroupsIndexMutex must be entered before the AllocationMap (that is the order used by e.g. TrackAlloc).
+	// So copy AllocationGroups before entering AllocationMap.LockAll().
+
+	// Note that this cache works even when !LLM_ALLOW_NAMES_TAGS, because the AllocationGroup's index when
+	// !LLM_ALLOW_NAMES_TAGS is defined to be equal to the ELLMTag, which is the compressed tag.
+	TArray<FAllocationGroup*, FDefaultLLMAllocator> AllocationGroupCached;
+	{
+		UE::TSharedLock Lock(LLMRef.AllocationGroupsIndexMutex);
+		AllocationGroupCached = *LLMRef.AllocationGroups;
+	}
+
 	AllocationMap.LockAll();
 	for (const FLLMAllocMap::FTuple& Tuple : AllocationMap)
 	{
-		bool bIncludeAllocation = true;
-		for (const FLLMTagSetAllocationFilter& Filter : Filters)
+		void* Ptr = Tuple.Key.GetPointer();
+		int64 SizeLow = static_cast<int64>(Tuple.Value1);
+		int64 SizeHigh = Tuple.Key.GetExtraData();
+		int64 Size = (SizeHigh << 32ull) | SizeLow;
+		const FLLMTracker::FLowLevelAllocInfo& Info = Tuple.Value2;
+		FAllocationGroup* AllocationGroup = nullptr;
+		int32 CompressedTagIndex = Info.GetCompressedTag();
+		if (AllocationGroupCached.IsValidIndex(CompressedTagIndex))
 		{
-			FLocalTagData* Data = nullptr;
-#if LLM_ALLOW_ASSETS_TAGS || LLM_ALLOW_UOBJECTCLASSES_TAGS
-			int32 CompressedTagIndex = static_cast<int32>(Tuple.Value2.GetCompressedTag(Filter.TagSet));
-#else
-			int32 CompressedTagIndex = static_cast<int32>(Tuple.Value2.GetCompressedTag());
-#endif
-			if(TagDataCached.IsValidIndex(CompressedTagIndex))
+			AllocationGroup = AllocationGroupCached[CompressedTagIndex];
+		}
+		Callback(Ptr, Size, AllocationGroup);
+	}
+	AllocationMap.UnlockAll();
+#endif // !UE_ONLY_USE_PLATFORM_TRACKER
+}
+
+void FLLMTracker::GetTagsNamesWithAmountFiltered(TMap<FName, uint64>& OutTagsNamesWithAmount,
+	ELLMTagSet TagSet /* = ELLMTagSet::None */, TArray<FLLMTagSetAllocationFilter>& Filters)
+{
+	// This function must be called within the UpdateLock, to avoid the possibility of AllocationGroups being deleted
+	// while the function is in progress.
+#if UE_ONLY_USE_PLATFORM_TRACKER
+	return;
+#else // !UE_ONLY_USE_PLATFORM_TRACKER
+
+	int32 IndexInActiveTags = LLMRef.GetIndexInFActiveTags(TagSet);
+	if (IndexInActiveTags == -1)
+	{
+		return;
+	}
+
+	TArray<int32> IndexOfFiltersInActiveTags;
+	IndexOfFiltersInActiveTags.SetNum(Filters.Num());
+	for (int32 FilterIndex = 0; FilterIndex < Filters.Num(); ++FilterIndex)
+	{
+		IndexOfFiltersInActiveTags[FilterIndex] = LLMRef.GetIndexInFActiveTags(Filters[FilterIndex].TagSet);
+	}
+
+	TMap<const FTagData*, int64, FDefaultSetLLMAllocator> Results;
+	{
+		UE::TSharedLock Lock(LLMRef.TagDataMutex);
+		Results.Reserve(LLMRef.TagDatas->Num());
+	}
+
+	EnumerateAllocations(
+		[&Filters, &IndexOfFiltersInActiveTags, IndexInActiveTags, &Results]
+		(void* Ptr, int64 Size, FAllocationGroup* AllocationGroup)
+		{
+			if (!AllocationGroup)
 			{
-#if LLM_ALLOW_ASSETS_TAGS || LLM_ALLOW_UOBJECTCLASSES_TAGS
-				Data = &TagDataCached[CompressedTagIndex];
-#else
-				if (Filter.TagSet == ELLMTagSet::None)
-				{
-					Data = &TagDataCached[CompressedTagIndex];
-				}
-#endif
+				return;
 			}
-			if (!Data || Data->Name != Filter.Name)
+			if (Size == 0)
 			{
-				bIncludeAllocation = false;
-				break;
+				return;
 			}
-		}
 
-		if (bIncludeAllocation)
-		{
-#if LLM_ALLOW_ASSETS_TAGS || LLM_ALLOW_UOBJECTCLASSES_TAGS
-			int32 CompressedTagIndex = Tuple.Value2.GetCompressedTag(TagSet);
-#else
-			int32 CompressedTagIndex = static_cast<int32>(Tuple.Value2.GetCompressedTag());
-#endif
-			FLocalTagData* Data = TagDataCached.IsValidIndex(CompressedTagIndex) ? &TagDataCached[CompressedTagIndex] : nullptr;
-			if (Data)
+			bool bIncludeAllocation = true;
+			for (int32 FilterIndex = 0; FilterIndex < Filters.Num(); ++FilterIndex)
 			{
-				Data->Size += Tuple.Value1;
+				const FLLMTagSetAllocationFilter& Filter = Filters[FilterIndex];
+				int32 IndexOfFilterInActiveTags = IndexOfFiltersInActiveTags[FilterIndex];
+
+				const FTagData* Data = nullptr;
+				if (IndexOfFilterInActiveTags != -1)
+				{
+					Data = AllocationGroup->GetActiveTags()[IndexOfFilterInActiveTags];
+				}
+				if (!Data || Data->GetName() != Filter.Name)
+				{
+					bIncludeAllocation = false;
+					break;
+				}
 			}
-		}
-	}
-	AllocationMap.UnlockAll();
-#endif // !UE_ONLY_USE_PLATFORM_TRACKER
 
-	for (const FLocalTagData& LocalTagData : TagDataCached)
+			if (bIncludeAllocation)
+			{
+				const FTagData* Data = AllocationGroup->GetActiveTags()[IndexInActiveTags];
+				if (Data)
+				{
+					Results.FindOrAdd(Data) += Size;
+				}
+			}
+		});
+
+	for (const TPair<const FTagData*, int64>& Pair : Results)
 	{
-		if (LocalTagData.Size != 0)
-		{
-			OutTagsNamesWithAmount.FindOrAdd(LocalTagData.Name, 0) += LocalTagData.Size;
-		}
+		check(Pair.Value != 0);
+		OutTagsNamesWithAmount.FindOrAdd(Pair.Key->GetName(), 0) += Pair.Value;
 	}
+#endif // !UE_ONLY_USE_PLATFORM_TRACKER
 }
 
 bool FLLMTracker::FindTagsForPtr(void* InPtr, TArray<const FTagData *, TInlineAllocator<static_cast<int32>(ELLMTagSet::Max)>>& OutTags) const
@@ -5849,23 +6288,27 @@ bool FLLMTracker::FindTagsForPtr(void* InPtr, TArray<const FTagData *, TInlineAl
 	{
 		return false;
 	}
+	FAllocationGroup* AllocationGroup = AllocInfoPtr.GetGroup(LLMRef);
+	if (!AllocationGroup)
+	{
+		return false;
+	}
+	const FActiveTags& ActiveTags = AllocationGroup->GetActiveTags();
 
 	OutTags.SetNumUninitialized(static_cast<int32>(ELLMTagSet::Max));
-#if LLM_ALLOW_ASSETS_TAGS || LLM_ALLOW_UOBJECTCLASSES_TAGS
+	int32 IndexInActiveTags = 0;
 	for (int32 TagSetAsInteger = 0; TagSetAsInteger < static_cast<int32>(ELLMTagSet::Max); TagSetAsInteger++)
 	{
-		// Calling GetTag with an ELLMTagSet that is compiled-out is illegal and will cause an assertion.
-		// We rely on IsTagSetActive to be false for compiled-out ELLMTagSet, and use that to avoid calling GetTag.
-		if (!LLMRef.IsTagSetActive(static_cast<ELLMTagSet>(TagSetAsInteger)))
+		if (LLMRef.ActiveSets[TagSetAsInteger])
 		{
-			continue;
+			OutTags[TagSetAsInteger] = ActiveTags[IndexInActiveTags];
+			++IndexInActiveTags;
+		}
+		else
+		{
+			OutTags[TagSetAsInteger] = nullptr;
 		}
-
-		OutTags[TagSetAsInteger] = AllocInfoPtr.GetTag(LLMRef, static_cast<ELLMTagSet>(TagSetAsInteger));
 	}
-#else
-	OutTags[static_cast<int32>(ELLMTagSet::None)] = AllocInfoPtr.GetTag(LLMRef);
-#endif
 
 	return true;
 #else
@@ -5919,35 +6362,40 @@ FLLMThreadState::FLLMThreadState()
 		PausedCounter[Index] = 0;
 	}
 
-	// By contract AllocationsLock must be entered around ClearAllocTypeAmounts, but multithreading contracts of that type
-	// are not required in constructors, which complete before the object can be used from multiple threads.
+	// By contract AllocationGroupsMutex must be entered around ClearAllocTypeAmounts, but multithreading contracts
+	// of that type are not required in constructors, which complete before the object can be used from multiple
+	// threads.
 	ClearAllocTypeAmounts();
 }
 
+FLLMThreadState::~FLLMThreadState()
+{
+	Clear();
+}
+
 void FLLMThreadState::Clear()
 {
-	FScopeLock Lock(&AllocationsLock);
+	UE::TUniqueLock ScopeLockAllocationGroups(ThreadAllocationGroupsMutex);
 
 	for (int32 TagSetAsInteger = 0; TagSetAsInteger < static_cast<int32>(ELLMTagSet::Max); TagSetAsInteger++)
 	{
 		TagStack[TagSetAsInteger].Empty();
 	}
-	Allocations.Empty();
+	for (TPair<FActiveTags, FAllocationGroupTrackingData>& Pair : ActiveTagsToTrackingData)
+	{
+		Pair.Value.ConditionalReset(true /* bReferenceTrack */);
+	}
+	ActiveTagsToTrackingData.Empty();
 	ClearAllocTypeAmounts();
 }
 
 void FLLMThreadState::PushTag(const FTagData* TagData, ELLMTagSet TagSet)
 {
-	FScopeLock Lock(&AllocationsLock);
-
-	// Push a tag.
 	TagStack[static_cast<int32>(TagSet)].Add(TagData);
 }
 
 void FLLMThreadState::PopTag(ELLMTagSet TagSet)
 {
-	FScopeLock Lock(&AllocationsLock);
-
 	LLMCheckf(TagStack[static_cast<int32>(TagSet)].Num() > 0,
 		TEXT("Called FLLMThreadState::PopTag without a matching Push (stack was empty on pop)"));
 	TagStack[static_cast<int32>(TagSet)].Pop(EAllowShrinking::No);
@@ -5963,215 +6411,317 @@ const FTagData* FLLMThreadState::GetTopTag(ELLMTagSet TagSet) const
 	return TagStack[static_cast<int32>(TagSet)].Last();
 }
 
-void FLLMThreadState::IncrTag(const FTagData* TagData, int64 Amount)
+void FLLMThreadState::BeginActiveSetsChange(
+	TArray<FAllocationGroupTrackingData, FDefaultLLMAllocator>& AllocationGroupTrackingDatas)
 {
-	// Caller is responsible for holding a lock on AllocationsLock.
-	FThreadTagSizeData& AllocationSize = Allocations.FindOrAdd(TagData->GetIndex());
-	AllocationSize.TagData = TagData;
-	AllocationSize.Size += Amount;
+	AllocationGroupTrackingDatas.Empty(ActiveTagsToTrackingData.Num());
+	for (TPair<FActiveTags, FAllocationGroupTrackingData>& Pair : ActiveTagsToTrackingData)
+	{
+		AllocationGroupTrackingDatas.Add(MoveTemp(Pair.Value));
+	}
+	ActiveTagsToTrackingData.Empty();
 }
 
-void FLLMThreadState::TrackAllocation(const void* Ptr, int64 Size, ELLMTracker Tracker, ELLMAllocType AllocType,
-	const FTagData* TagData, const FTagData* AssetTagData, const FTagData* AssetClassTagData, const FTagData* UObjectClassesTagData, bool bTrackInMemPro)
+void FLLMThreadState::EndActiveSetsChange(const FActiveTags& NewDefaults,
+	TArray<FAllocationGroupTrackingData, FDefaultLLMAllocator>& AllocationGroupTrackingDatas)
 {
-	FScopeLock Lock(&AllocationsLock);
-
-	AllocTypeAmounts[static_cast<int32>(AllocType)] += Size;
-
-	if (TagData)
-	{
-		IncrTag(TagData, Size);
-	}
-#if LLM_ALLOW_ASSETS_TAGS
-	if (AssetTagData)
-	{
-		IncrTag(AssetTagData, Size);
-	}
-	if (AssetClassTagData)
+	ActiveTagsToTrackingData.Reserve(AllocationGroupTrackingDatas.Num());
+	for (FAllocationGroupTrackingData& TrackingData : AllocationGroupTrackingDatas)
 	{
-		IncrTag(AssetClassTagData, Size);
+		FAllocationGroup* Group = TrackingData.GetGroup();
+		ActiveTagsToTrackingData.Add(Group->GetActiveTags(), MoveTemp(TrackingData));
 	}
-#endif // LLM_ALLOW_ASSETS_TAGS
+}
 
-#if LLM_ALLOW_UOBJECTCLASSES_TAGS
-	if (UObjectClassesTagData)
+FActiveTags FLLMThreadState::GetActiveTags(FLLMGlobals& LLMRef) const
+{
+	// This function must always be called outside of the ThreadAllocationsGroupRWLock, because we may enter that lock
+	// during some helper functions.
+	FActiveTags Result;
+	uint8 IndexInActiveTags = 0;
+	for (uint8 TagSetIndex = 0; TagSetIndex < static_cast<uint8>(ELLMTagSet::Max); ++TagSetIndex)
 	{
-		IncrTag(UObjectClassesTagData, Size);
-	}
-#endif // LLM_ALLOW_UOBJECTCLASSES_TAGS
+		ELLMTagSet TagSet = static_cast<ELLMTagSet>(TagSetIndex);
+		if (!LLMRef.ActiveSets[TagSetIndex])
+		{
+			continue;
+		}
 
-	ELLMTag EnumTag = TagData ? TagData->GetContainingEnum() : ELLMTag(0);
-	if (Tracker == ELLMTracker::Default)
-	{
-		FPlatformMemory::OnLowLevelMemory_Alloc(Ptr, static_cast<uint64>(Size), static_cast<uint64>(EnumTag));
-	}
+		const FTagData* TagData = GetTopTag(TagSet);
+		if (TagData == nullptr)
+		{
+			TagData = LLMRef.FindOrAddDefaultTagData(TagSet);
+		}
 
-#if MEMPRO_ENABLED
-	if (FMemProProfiler::IsTrackingTag(EnumTag) && bTrackInMemPro)
-	{
-		MEMPRO_TRACK_ALLOC(const_cast<void*>(Ptr), static_cast<size_t>(Size));
-	}
+#if !LLM_ALLOW_NAMES_TAGS
+		if (TagSet == ELLMTagSet::None)
+		{
+			// When full tags are disabled, we instead store the top-level enumtag parent of the allocation's tag.
+			TagData = TagData->GetContainingEnumTagData();
+		}
 #endif
+
+		Result[IndexInActiveTags] = TagData;
+		++IndexInActiveTags;
+	}
+	return Result;
 }
 
-void FLLMThreadState::TrackFree(const void* Ptr, int64 Size, ELLMTracker Tracker, ELLMAllocType AllocType,
-	const FTagData* TagData, const FTagData* AssetTagData, const FTagData* AssetClassTagData, const FTagData* UObjectClassesTagData, bool bTrackInMemPro)
+FAllocationGroupTrackingData& FLLMThreadState::FindOrAddTrackingData(const FActiveTags& ActiveTags, FAllocationGroup* AllocationGroup)
 {
-	FScopeLock Lock(&AllocationsLock);
+	// Caller must hold a WriteScopeLock on ThreadAllocationGroupsMutex
+	FAllocationGroupTrackingData& TrackingData = ActiveTagsToTrackingData.FindOrAdd(ActiveTags);
+	TrackingData.ConditionalInitialize(AllocationGroup, true /* bReferenceTrack */);
+	return TrackingData;
+}
 
-	AllocTypeAmounts[static_cast<int32>(AllocType)] -= Size;
-	if (TagData)
-	{
-		IncrTag(TagData, -Size);
-	}
-#if LLM_ALLOW_ASSETS_TAGS
-	if (AssetTagData)
-	{
-		IncrTag(AssetTagData, -Size);
-	}
-	if (AssetClassTagData)
+void FLLMThreadState::TrackAllocation(const void* Ptr, int64 Size, FLLMGlobals& LLMRef, ELLMTracker Tracker,
+	ELLMAllocType AllocType, bool bTrackInMemPro, FAllocationGroup*& OutGroup)
+{
+	FActiveTags ActiveTags = GetActiveTags(LLMRef);
 	{
-		IncrTag(AssetClassTagData, -Size);
-	}
-#endif // LLM_ALLOW_ASSETS_TAGS
+		UE::TUniqueLock WriteScopeLockThread(ThreadAllocationGroupsMutex);
 
-#if LLM_ALLOW_UOBJECTCLASSES_TAGS
-	if (UObjectClassesTagData)
-	{
-		IncrTag(UObjectClassesTagData, -Size);
+		AllocTypeAmounts[static_cast<int32>(AllocType)] += Size;
+
+		const FTagData* SystemsTagData = ActiveTags.GetSystemsTagData();
+		ELLMTag EnumTag = SystemsTagData ? SystemsTagData->GetContainingEnum() : ELLMTag(0);
+		if (Tracker == ELLMTracker::Default)
+		{
+			FPlatformMemory::OnLowLevelMemory_Alloc(Ptr, static_cast<uint64>(Size), static_cast<uint64>(EnumTag));
+		}
+#if MEMPRO_ENABLED
+		if (FMemProProfiler::IsTrackingTag(EnumTag) && bTrackInMemPro)
+		{
+			MEMPRO_TRACK_ALLOC(const_cast<void*>(Ptr), static_cast<size_t>(Size));
+		}
+#endif
+
+		FAllocationGroupTrackingData* TrackingData = ActiveTagsToTrackingData.Find(ActiveTags);
+		if (TrackingData)
+		{
+			TrackingData->TrackAllocOrFree(Size, 1);
+			OutGroup = TrackingData->GetGroup();
+			return;
+		}
 	}
-#endif // LLM_ALLOW_UOBJECTCLASSES_TAGS
 
-	ELLMTag EnumTag = TagData ? TagData->GetContainingEnum() : ELLMTag(0);
-	if (Tracker == ELLMTracker::Default)
 	{
-		FPlatformMemory::OnLowLevelMemory_Free(Ptr, static_cast<uint64>(Size), static_cast<uint64>(EnumTag));
+		UE::TSharedLock ReadScopeLockGlobal(LLMRef.AllocationGroupsMutex);
+		FAllocationGroup** AllocationGroup = LLMRef.ActiveTagsToAllocationGroup->Find(ActiveTags);
+		if (AllocationGroup)
+		{
+			UE::TUniqueLock WriteScopeLockThread(ThreadAllocationGroupsMutex);
+			FAllocationGroupTrackingData& TrackingData = FindOrAddTrackingData(ActiveTags, *AllocationGroup);
+			TrackingData.TrackAllocOrFree(Size, 1);
+			OutGroup = *AllocationGroup;
+			return;
+		}
 	}
 
-#if MEMPRO_ENABLED
-	if (FMemProProfiler::IsTrackingTag(EnumTag) && bTrackInMemPro)
 	{
-		MEMPRO_TRACK_FREE(const_cast<void*>(Ptr));
+		UE::TUniqueLock WriteScopeLockGlobal(LLMRef.AllocationGroupsMutex);
+		FAllocationGroup* AllocationGroup = LLMRef.FindOrAddAllocationGroup(ActiveTags);
+		{
+			UE::TUniqueLock WriteScopeLockThread(ThreadAllocationGroupsMutex);
+			FAllocationGroupTrackingData& TrackingData = FindOrAddTrackingData(ActiveTags, AllocationGroup);
+			TrackingData.TrackAllocOrFree(Size, 1);
+			OutGroup = AllocationGroup;
+			return;
+		}
 	}
-#endif
 }
 
-void FLLMThreadState::TrackMemory(int64 Amount, ELLMTracker Tracker, ELLMAllocType AllocType, const FTagData* TagData)
+void FLLMThreadState::TrackFree(const void* Ptr, int64 Size, ELLMTracker Tracker, ELLMAllocType AllocType,
+	FAllocationGroup* AllocationGroup, bool bTrackInMemPro)
 {
-	FScopeLock Lock(&AllocationsLock);
-	AllocTypeAmounts[static_cast<int32>(AllocType)] += Amount;
-	IncrTag(TagData, Amount);
+	{
+		UE::TUniqueLock WriteScopeLockThread(ThreadAllocationGroupsMutex);
 
-	// TODO: Need to expose TrackMemory to Platform-specific trackers and FMemPropProfiler
-}
+		AllocTypeAmounts[static_cast<int32>(AllocType)] -= Size;
 
-void FLLMThreadState::TrackMoved(const void* Dest, const void* Source, int64 Size, ELLMTracker Tracker,
-	const FTagData* TagData)
-{
-	// Update external memory trackers (ideally would want a proper 'move' option on these).
-	ELLMTag EnumTag = TagData ? TagData->GetContainingEnum() : ELLMTag(0);
-	if (Tracker == ELLMTracker::Default)
-	{
-		FPlatformMemory::OnLowLevelMemory_Free(Source, static_cast<uint64>(Size), static_cast<uint64>(EnumTag));
-		FPlatformMemory::OnLowLevelMemory_Alloc(Dest, static_cast<uint64>(Size), static_cast<uint64>(EnumTag));
-	}
+		ELLMTag EnumTag = ELLMTag(0);
+		if (AllocationGroup)
+		{
+			EnumTag = AllocationGroup->GetActiveTags().GetSystemsTagData()->GetContainingEnum();
+		}
+		if (Tracker == ELLMTracker::Default)
+		{
+			FPlatformMemory::OnLowLevelMemory_Free(Ptr, static_cast<uint64>(Size), static_cast<uint64>(EnumTag));
+		}
 
 #if MEMPRO_ENABLED
-	if (FMemProProfiler::IsTrackingTag(EnumTag))
+		if (FMemProProfiler::IsTrackingTag(EnumTag) && bTrackInMemPro)
+		{
+			MEMPRO_TRACK_FREE(const_cast<void*>(Ptr));
+		}
+#endif
+
+		if (!AllocationGroup)
+		{
+			return;
+		}
+
+		const FActiveTags& ActiveTags = AllocationGroup->GetActiveTags();
+		FAllocationGroupTrackingData* TrackingData = ActiveTagsToTrackingData.Find(ActiveTags);
+		if (TrackingData)
+		{
+			TrackingData->TrackAllocOrFree(-Size, -1);
+			return;
+		}
+	}
+
 	{
-		MEMPRO_TRACK_FREE(const_cast<void*>(Source));
-		MEMPRO_TRACK_ALLOC(const_cast<void*>(Dest), static_cast<size_t>(Size));
+		UE::TUniqueLock WriteScopeLockThread(ThreadAllocationGroupsMutex);
+		const FActiveTags& ActiveTags = AllocationGroup->GetActiveTags();
+		FAllocationGroupTrackingData& TrackingData = FindOrAddTrackingData(ActiveTags, AllocationGroup);
+		TrackingData.TrackAllocOrFree(-Size, -1);
+		return;
 	}
-#endif
+
 }
 
-void FLLMThreadState::PropagateChildSizesToParents()
+void FLLMThreadState::TrackMemory(int64 Amount, FLLMGlobals& LLMRef, ELLMTracker Tracker, ELLMAllocType AllocType, const FTagData* TagData)
 {
-	FScopeLock Lock(&AllocationsLock);
+#if !LLM_ALLOW_NAMES_TAGS
+	// When full tags are disabled, we instead store the top-level enumtag parent of the allocation's tag.
+	TagData = TagData ? TagData->GetContainingEnumTagData() : nullptr;
+#endif
 
-	// Make sure all parents of any TagDatas in the Allocations are also present.
-	FConstTagDataArray ParentsToAdd;
-	for (const TPair<int32, FThreadTagSizeData>& It : Allocations)
+	// TrackMemory only tracks memory on ELLMTagSet::None; all other tagsets are set to default.
+	FActiveTags ActiveTags;
+	int32 IndexInActiveTags = 0;
+	for (uint8 TagSetIndex = 0; TagSetIndex < static_cast<uint8>(ELLMTagSet::Max); ++TagSetIndex)
 	{
-		const FTagData* TagData = It.Value.TagData;
-		const FTagData* ParentData = TagData->GetParent();
-		while (ParentData && !Allocations.Contains(ParentData->GetIndex()))
+		if (LLMRef.ActiveSets[TagSetIndex])
 		{
-			ParentsToAdd.Add(ParentData);
-			ParentData = ParentData->GetParent();
+			ELLMTagSet TagSet = static_cast<ELLMTagSet>(TagSetIndex);
+			if (TagSet == ELLMTagSet::None)
+			{
+				ActiveTags[IndexInActiveTags] = TagData;
+			}
+			else
+			{
+				ActiveTags[IndexInActiveTags] = LLMRef.FindOrAddDefaultTagData(TagSet);
+			}
+			++IndexInActiveTags;
 		}
 	}
-	for (const FTagData* TagData : ParentsToAdd)
+	LLMCheck(ActiveTags.GetSystemsTagData() == TagData);
+
 	{
-		FThreadTagSizeData& Info = Allocations.FindOrAdd(TagData->GetIndex());
-		Info.TagData = TagData;
+		UE::TUniqueLock WriteScopeLockThread(ThreadAllocationGroupsMutex);
+
+		AllocTypeAmounts[static_cast<int32>(AllocType)] += Amount;
+
+		// TODO: Need to expose TrackMemory to Platform-specific trackers and FMemPropProfiler
+
+		FAllocationGroupTrackingData* TrackingData = ActiveTagsToTrackingData.Find(ActiveTags);
+		if (TrackingData)
+		{
+			TrackingData->TrackAllocOrFree(Amount, 0);
+			return;
+		}
+	}
+
+	{
+		UE::TSharedLock ReadScopeLockGlobal(LLMRef.AllocationGroupsMutex);
+		FAllocationGroup** AllocationGroup = LLMRef.ActiveTagsToAllocationGroup->Find(ActiveTags);
+		if (AllocationGroup)
+		{
+			UE::TUniqueLock WriteScopeLockThread(ThreadAllocationGroupsMutex);
+			FAllocationGroupTrackingData& TrackingData = FindOrAddTrackingData(ActiveTags, *AllocationGroup);
+			TrackingData.TrackAllocOrFree(Amount, 0);
+			return;
+		}
 	}
 
-	// Tags are sorted topologically from parent to child, so we can accumulate children into parents recursively
-	// by reverse iterating the map.
-	for (FThreadTagSizeMap::TConstReverseIterator It(Allocations); It; ++It)
 	{
-		const FThreadTagSizeData& Info = It->Value;
-		const FTagData* ParentData = Info.TagData->GetParent();
-		if (Info.Size && ParentData)
+		UE::TUniqueLock WriteScopeLockGlobal(LLMRef.AllocationGroupsMutex);
+		FAllocationGroup* AllocationGroup = LLMRef.FindOrAddAllocationGroup(ActiveTags);
 		{
-			Allocations.FindChecked(ParentData->GetIndex()).Size += Info.Size;
+			UE::TUniqueLock WriteScopeLockThread(ThreadAllocationGroupsMutex);
+			FAllocationGroupTrackingData& TrackingData = FindOrAddTrackingData(ActiveTags, AllocationGroup);
+			TrackingData.TrackAllocOrFree(Amount, 0);
+			return;
 		}
 	}
 }
 
-void FLLMThreadState::OnTagsResorted(FTagDataArray& OldTagDatas)
+void FLLMThreadState::TrackMoved(const void* Dest, const void* Source, int64 Size, ELLMTracker Tracker,
+	FAllocationGroup* AllocationGroup)
 {
-	FScopeLock Lock(&AllocationsLock);
-	TArray<FThreadTagSizeData, FDefaultLLMAllocator> AllocationDatas;
-	AllocationDatas.Reserve(Allocations.Num());
-	for (const TPair<int32, FThreadTagSizeData>& It : Allocations)
+	// Update external memory trackers (ideally would want a proper 'move' option on these).
+	// We have to enter the mutex because we are not allowed to dereference AllocationGroup without being in the mutex.
+	UE::TUniqueLock WriteScopeLockThread(ThreadAllocationGroupsMutex);
+
+	ELLMTag EnumTag = ELLMTag(0);
+	if (AllocationGroup)
 	{
-		AllocationDatas.Add(It.Value);
+		EnumTag = AllocationGroup->GetActiveTags().GetSystemsTagData()->GetContainingEnum();
+	}
+
+	if (Tracker == ELLMTracker::Default)
+	{
+		FPlatformMemory::OnLowLevelMemory_Free(Source, static_cast<uint64>(Size), static_cast<uint64>(EnumTag));
+		FPlatformMemory::OnLowLevelMemory_Alloc(Dest, static_cast<uint64>(Size), static_cast<uint64>(EnumTag));
 	}
-	Allocations.Reset();
-	for (const FThreadTagSizeData& AllocationData : AllocationDatas)
+
+#if MEMPRO_ENABLED
+	if (FMemProProfiler::IsTrackingTag(EnumTag))
 	{
-		Allocations.Add(AllocationData.TagData->GetIndex(), AllocationData);
+		MEMPRO_TRACK_FREE(const_cast<void*>(Source));
+		MEMPRO_TRACK_ALLOC(const_cast<void*>(Dest), static_cast<size_t>(Size));
 	}
+#endif
 }
 
-void FLLMThreadState::LockTags(bool bLock)
+void FLLMThreadState::LockSharedData(bool bLock)
 {
 	if (bLock)
 	{
-		AllocationsLock.Lock();
+		ThreadAllocationGroupsMutex.Lock();
 	}
 	else
 	{
-		AllocationsLock.Unlock();
+		ThreadAllocationGroupsMutex.Unlock();
 	}
 }
 
-void FLLMThreadState::FetchAndClearTagSizes(FTrackerTagSizeMap& TagSizes, int64* InAllocTypeAmounts,
-	bool bTrimAllocations)
+void FLLMThreadState::FetchAndClearTagSizes(FMapActiveTagsToAllocationGroupTrackingData& OutDatas,
+	int64* OutAllocTypeAmounts, EPruningLevel PruningLevel)
 {
-	FScopeLock Lock(&AllocationsLock);
-	for (TPair<int32, FThreadTagSizeData>& Item : Allocations)
 	{
-		FThreadTagSizeData& ThreadInfo = Item.Value;
-		if (ThreadInfo.Size)
+		// Move the Size/AllocCount data from the thread-shared container into a local copy, and prune the thread-shared
+		// container for any recently unused entries.
+		UE::TUniqueLock ScopeLock(ThreadAllocationGroupsMutex);
+		for (auto TrackingDataIter(ActiveTagsToTrackingData.CreateIterator()); TrackingDataIter; ++TrackingDataIter)
 		{
-			const FTagData* TagData = ThreadInfo.TagData;
-			FTrackerTagSizeData& TrackerInfo = TagSizes.FindOrAdd(TagData);
-			TrackerInfo.Size += ThreadInfo.Size;
-			ThreadInfo.Size = 0;
+			FAllocationGroupTrackingData& TrackingData = TrackingDataIter.Value();
+			FAllocationGroupTrackingData& OutData = OutDatas.FindOrAdd(TrackingDataIter.Key());
+			TrackingData.DetachSizeAndActiveAllocs(OutData, PruningLevel);
+			if (PruningLevel > EPruningLevel::None
+				&& TrackingData.GetUpdatesSinceLastAllocationChange() >= PruneUpdatesForTrackingData)
+			{
+				TrackingData.ConditionalReset(true /* bReferenceTrack */);
+				TrackingDataIter.RemoveCurrent();
+			}
+		}
+		if (PruningLevel > EPruningLevel::None)
+		{
+			if (PruningLevel >= EPruningLevel::Max ||
+				ActiveTagsToTrackingData.Max() > (int32)((float) ActiveTagsToTrackingData.Num() * PrunePopulationRatioForTrackingData))
+			{
+				ActiveTagsToTrackingData.Shrink();
+			}
 		}
-	}
-	if (bTrimAllocations)
-	{
-		Allocations.Empty();
-	}
 
-	for (int32 Index = 0; Index < static_cast<int32>(ELLMAllocType::Count); ++Index)
-	{
-		InAllocTypeAmounts[Index] += AllocTypeAmounts[Index];
-		AllocTypeAmounts[Index] = 0;
+		// Also Move the AllocTypeAmounts data
+		for (int32 Index = 0; Index < static_cast<int32>(ELLMAllocType::Count); ++Index)
+		{
+			OutAllocTypeAmounts[Index] += AllocTypeAmounts[Index];
+			AllocTypeAmounts[Index] = 0;
+		}
 	}
 }
 
@@ -6185,8 +6735,7 @@ void SetMemoryStatByFName(FName Name, int64 Amount)
 
 void FLLMThreadState::ClearAllocTypeAmounts()
 {
-	// Caller is responsible for holding a lock on AllocationsLock.
-
+	// Caller is responsible for holding a lock on ThreadAllocationGroupsMutex or global AllocationGroupsMutex.
 	for (int32 Index = 0; Index < static_cast<int32>(ELLMAllocType::Count); ++Index)
 	{
 		AllocTypeAmounts[Index] = 0;
@@ -6198,6 +6747,7 @@ FLLMCsvWriter::FLLMCsvWriter()
 	: Archive(nullptr)
 	, LastWriteTime(FPlatformTime::Seconds())
 	, WriteCount(0)
+	, HeaderMaxSize(-1)
 	, bRegisteredFlushDelegate(false)
 {
 }
@@ -6316,7 +6866,35 @@ bool FLLMCsvWriter::CreateArchive()
 	LLMCheck(Archive);
 
 	// Create space for column titles that are filled in as we get them.
-	Write(FString::ChrN(CVarLLMHeaderMaxSize.GetValueOnAnyThread(), ' '));
+	int32 MaxRequestedSize = CVarLLMHeaderMaxSize.GetValueOnAnyThread();
+	if (MaxRequestedSize > 0)
+	{
+		HeaderMaxSize = MaxRequestedSize;
+	}
+	else
+	{
+		HeaderMaxSize = 0;
+		constexpr int32 MaxNameStrLen = 256;
+		TStringBuilder<MaxNameStrLen> TagNameStr;
+		FLLMGlobals& LLMRef = FLLMGlobals::GetInner();
+		// Find the name size of all reportable TagDatas. This will not include TagDatas that have not yet been created,
+		// but it will gather all of the compiled-in ELLMTag and FName tags. We only allow reporting
+		// ELLMTagSet::None, and not many dynamic tags exist in ELLMTagSet::None, so this will reserve space for most of
+		// them.
+		UE::TSharedLock Lock(LLMRef.TagDataMutex);
+		for (FTagData* TagData : (*LLMRef.TagDatas))
+		{
+			if (TagData->IsReportable())
+			{
+				TagNameStr.Reset();
+				// We have to avoid allocating memory in this loop because we hold an LLM lock.
+				TagData->GetDisplayPath(TagNameStr, MaxNameStrLen - 1);
+				HeaderMaxSize += TagNameStr.Len() + 1;
+			}
+		}
+	}
+
+	Write(FString::ChrN(HeaderMaxSize, ' '));
 	Write(TEXT("\n"));
 
 	return true;
@@ -6376,7 +6954,7 @@ void FLLMCsvWriter::WriteHeader(const FTagData* OverrideTrackedTotalTagData, con
 	}
 
 	int64 ColumnTitleTotalSize = Archive->Tell();
-	if (ColumnTitleTotalSize >= CVarLLMHeaderMaxSize.GetValueOnAnyThread())
+	if (ColumnTitleTotalSize >= HeaderMaxSize)
 	{
 		UE_LOG(LogHAL, Error,
 			TEXT("LLM column titles have overflowed, LLM CSM data will be corrupted. Increase CVarLLMHeaderMaxSize > %" INT64_FMT),
@@ -6740,4 +7318,4 @@ FORCENOINLINE void* FArraySlackTrackingHeader::Realloc(void* Ptr, int64 Count, u
 }
 #endif
 
-#endif  // #else .. #if ENABLE_LOW_LEVEL_MEM_TRACKER
+#endif  // #else .. #if ENABLE_LOW_LEVEL_MEM_TRACKER
\ No newline at end of file
diff --git a/Engine/Source/Runtime/Core/Private/HAL/LLM/LLMPrivate.h b/Engine/Source/Runtime/Core/Private/HAL/LLM/LLMPrivate.h
--- a/Engine/Source/Runtime/Core/Private/HAL/LLM/LLMPrivate.h
+++ b/Engine/Source/Runtime/Core/Private/HAL/LLM/LLMPrivate.h
@@ -10,12 +10,14 @@
 #include "Containers/Set.h"
 #include "Containers/SortedMap.h"
 #include "HAL/CriticalSection.h"
+#include "HAL/LLM/AllocationGroup.h"
 #include "HAL/LLM/LLMPrivateDefines.h"
 #include "HAL/LLM/MemoryUtils.h"
 #include "HAL/PlatformMutex.h"
 #include "Misc/Optional.h"
 #include "Misc/ScopeRWLock.h"
 #include "ProfilingDebugging/CsvProfiler.h"
+#include "Templates/Function.h"
 
 void GlobalRegisterTagDeclaration(FLLMTagDeclaration& TagDeclaration);
 
@@ -272,6 +274,7 @@ private:
 	FArchive* Archive;
 	double LastWriteTime;
 	int32 WriteCount;
+	int32 HeaderMaxSize;
 	ELLMTracker Tracker;
 	bool bRegisteredFlushDelegate;
 };
@@ -320,36 +323,49 @@ class FLLMThreadState
 {
 public:
 	FLLMThreadState();
+	~FLLMThreadState();
 
 	void Clear();
 
 	void PushTag(const FTagData* TagData, ELLMTagSet TagSet);
 	void PopTag(ELLMTagSet TagSet);
 	const FTagData* GetTopTag(ELLMTagSet TagSet) const;
-	void TrackAllocation(const void* Ptr, int64 Size, ELLMTracker Tracker, ELLMAllocType AllocType,
-		const FTagData* TagData, const FTagData* AssetTagData, const FTagData* AssetClassTagData, const FTagData* UObjectClassesTagData, bool bTrackInMemPro);
+	void TrackAllocation(const void* Ptr, int64 Size, FLLMGlobals& LLMRef, ELLMTracker Tracker,
+		ELLMAllocType AllocType, bool bTrackInMemPro, FAllocationGroup*& OutGroup);
+	FAllocationGroupTrackingData& FindOrAddTrackingData(const FActiveTags& ActiveTags,
+		FAllocationGroup* AllocationGroup);
 	void TrackFree(const void* Ptr, int64 Size, ELLMTracker Tracker, ELLMAllocType AllocType,
-		const FTagData* TagData, const FTagData* AssetTagData, const FTagData* AssetClassTagData, const FTagData* UObjectClassesTagData, bool bTrackInMemPro);
-	void TrackMemory(int64 Amount, ELLMTracker Tracker, ELLMAllocType AllocType,
+		FAllocationGroup* AllocationGroup, bool bTrackInMemPro);
+	void TrackMemory(int64 Amount, FLLMGlobals& LLMRef, ELLMTracker Tracker, ELLMAllocType AllocType,
 		const FTagData* TagData);
-	void TrackMoved(const void* Dest, const void* Source, int64 Size, ELLMTracker Tracker, const FTagData* TagData);
-	void IncrTag(const FTagData* Tag, int64 Amount);
+	void TrackMoved(const void* Dest, const void* Source, int64 Size, ELLMTracker Tracker,
+		FAllocationGroup* AllocationGroup);
 
-	void PropagateChildSizesToParents();
-	void OnTagsResorted(FTagDataArray& OldTagDatas);
-	void LockTags(bool bLock);
-	void FetchAndClearTagSizes(FTrackerTagSizeMap& TagSizes, int64* InAllocTypeAmounts, bool bTrimAllocations);
+	void LockSharedData(bool bLock);
+	void FetchAndClearTagSizes(FMapActiveTagsToAllocationGroupTrackingData& OutDatas, int64* OutAllocTypeAmounts,
+		EPruningLevel PruningLevel);
 
 	void ClearAllocTypeAmounts();
 
-	// AllocationsLock must be held around read/write of Allocations and AllocTypeAmounts, because these are read/write
-	// from the Update thread. All other data on FLLMThreadState are accessed only from the thread that owns the
+	FActiveTags GetActiveTags(FLLMGlobals& LLMRef) const;
+	void BeginActiveSetsChange(
+		TArray<FAllocationGroupTrackingData, FDefaultLLMAllocator>& AllocationGroupTrackingDatas);
+	void EndActiveSetsChange(const FActiveTags& NewDefaults,
+		TArray<FAllocationGroupTrackingData, FDefaultLLMAllocator>& AllocationGroupTrackingDatas);
+
+	// Either LLMRef.AllocationGroupsMutex or this->ThreadAllocationGroupsMutex must be held around read/write
+	// of ActiveTagsToTrackingData and AllocTypeAmounts, because the Update thread enters a write lock on those
+	// two locks (AllocationGroupsMutex entered first) when read/writing ActiveTagsToTrackingData and
+	// AllocTypeAmounts from the update thread. ThreadAllocationGroupsMutex should be preferred when possible to avoid
+	// blocking other threads. All other data on FLLMThreadState are accessed only from the thread that owns the
 	// FLLMThreadState and do not require a lock.
-	// AllocationsLock is also used to guarantee our contract provided to FLLMTracker::LockAllThreadTags: We must not
-	// dereference any FTagData* held by the FLLMThreadState unless we have entered AllocationsLock.
+	// These locks are also used to implement our contract that we guarantee for the Update thread's modification
+	// of TagDatas and AllocationGroups which it does while holding FLLMTracker::LockAllThreadsSharedData: we must not
+	// dereference any FTagData* or FAllocationGroup* held by the FLLMThreadState unless we have entered the global
+	// AllocationGroupsMutex.
 	FConstTagDataArray TagStack[static_cast<int32>(ELLMTagSet::Max)];
-	FThreadTagSizeMap Allocations;
-	FCriticalSection AllocationsLock;
+	FMapActiveTagsToAllocationGroupTrackingData ActiveTagsToTrackingData;
+	UE::FPlatformRecursiveMutex ThreadAllocationGroupsMutex;
 
 	int8 PausedCounter[(int32)ELLMAllocType::Count];
 	int64 AllocTypeAmounts[(int32)ELLMAllocType::Count];
@@ -410,47 +426,18 @@ public:
 	struct FLowLevelAllocInfo
 	{
 	public:
-		FLowLevelAllocInfo();
-
-		static int32 GetIndexFromTagSet(ELLMTagSet InTagSet);
-
-		void SetTag(const FTagData* InTag, FLLMGlobals& InLLMRef
-#if LLM_ALLOW_ASSETS_TAGS || LLM_ALLOW_UOBJECTCLASSES_TAGS
-			, ELLMTagSet InTagSet = ELLMTagSet::None
-#endif
-		);
-
-		const FTagData* GetTag(FLLMGlobals& InLLMRef
-#if LLM_ALLOW_ASSETS_TAGS || LLM_ALLOW_UOBJECTCLASSES_TAGS
-			, ELLMTagSet InTagSet = ELLMTagSet::None
-#endif
-		) const;
-
-#if LLM_ENABLED_FULL_TAGS
-		int32 GetCompressedTag(ELLMTagSet InTagSet = ELLMTagSet::None) const;
-		void SetCompressedTag(int32 InTag, ELLMTagSet InTagSet = ELLMTagSet::None);
-#else
-		ELLMTag GetCompressedTag();
-		void SetCompressedTag(ELLMTag InTag);
-#endif
+		void SetGroup(FAllocationGroup* AllocationGroup);
+		FAllocationGroup* GetGroup(FLLMGlobals& InLLMRef) const;
+		int32 GetCompressedTag() const;
 
 	private:
-#if LLM_ENABLED_FULL_TAGS
-		// Even with arbitrary tags we are still partially compressed - the allocation records the tag's index
-		// (4 bytes) rather than the full tag pointer (8 bytes) or the tag's name (12 bytes).
-#if !LLM_ALLOW_ASSETS_TAGS && !LLM_ALLOW_UOBJECTCLASSES_TAGS
-		int32 Tag = 0;
-#else
-		int32 Tag[LLM_TAGSET_MAX];
-
-		static_assert(LLM_TAGSET_MAX <= 3,
-			"Incrementing the maximum amount of TagSet enums that can be defined at the same time adds 4 bytes of LLM overhead memory FOR EVERY ACTIVE ALLOCATION (e.g. 10 million * 4 == 40 million bytes). "
-			"Do not add them without measuring the cost of the overhead in the cook and runtime of large projects.");
-#endif // LLM_ALLOW_ASSETS_TAGS
+#if LLM_ALLOW_NAMES_TAGS
+		// Even with name-based tags we are still partially compressed - the allocation records the AllocationGroup's
+		// index (4 bytes) rather than the full AllocationGroup pointer (8 bytes).
+		int32 Group = -1;
 #else
 		ELLMTag Tag = ELLMTag::Untagged;
 #endif
-		static const int32 InvalidCompressedTagValue = -1;
 	};
 
 	typedef LLMMap<PointerKey, uint32, FLowLevelAllocInfo, LLMNumAllocsType> FLLMAllocMap;  // pointer, size, info, Capacity SizeType
@@ -459,8 +446,13 @@ public:
 	void Update();
 	void UpdateThreads();
 	void OnTagsResorted(FTagDataArray& OldTagDatas);
-	void LockAllThreadTags(bool bLock);
+	void LockAllThreadsSharedData(bool bLock);
+	void BeginActiveSetsChange(
+		TArray<TArray<FAllocationGroupTrackingData, FDefaultLLMAllocator>, FDefaultLLMAllocator>& ThreadStateAllocationGroupTrackingDatas);
+	void EndActiveSetsChange(const FActiveTags& NewDefaults,
+		TArray<TArray<FAllocationGroupTrackingData, FDefaultLLMAllocator>, FDefaultLLMAllocator>& ThreadStateAllocationGroupTrackingDatas);
 
+	void FetchAndClearTagSizes(bool bUpdatePruning);
 	void CaptureTagSnapshot();
 	void ClearTagSnapshot();
 
@@ -486,12 +478,19 @@ public:
 		return TrackedTotal;
 	}
 
+	/**
+	 * Slow function for callers that need to examine every allocation. Thread-safe. Call outside of any LLM locks, and
+	 * in callback do not call LLM functions or allocate normally (use LLM internal allocator instead).
+	 */
+	void EnumerateAllocations(TFunctionRef<void(void* Ptr, int64 Size, FAllocationGroup* Group)> Callback);
+
 protected:
 	FLLMThreadState* GetOrCreateState();
 	FLLMThreadState* GetState();
-	void TrackAllocation(const void* Ptr, int64 Size, const FTagData* ActiveTagData, ELLMAllocType AllocType, FLLMThreadState* State, bool bTrackInMemPro);
+	void TrackAllocationInActiveTags(const void* Ptr, int64 Size, ELLMAllocType AllocType, FLLMThreadState* State, bool bTrackInMemPro);
 	void TrackMemoryOfActiveTag(int64 Amount, const FTagData* TagData, ELLMAllocType AllocType, FLLMThreadState* State);
-
+	void FetchAndClearTagSizes_UpdateGroupAllocs(FAllocationGroup& AllocationGroup,
+		const FAllocationGroupTrackingData& TrackingData, EPruningLevel PruningLevel);
 	FLLMGlobals& LLMRef;
 
 	ELLMTracker Tracker;
@@ -593,7 +592,10 @@ public:
 	void FinishInitializeInner();
 
 	// Functions for use by types within LLMPrivate
+	/** Accessor for the singleton FLLMGlobals for LLM code that does not already have a reference to it. */
 	static FLLMGlobals& GetInner();
+	/** Return the index in FActiveTags structures that corresponds to TagSet, or -1 if TagSet is not enabled. */
+	int32 GetIndexInFActiveTags(ELLMTagSet TagSet);
 
 private:
 	FLLMGlobals();
@@ -619,12 +621,17 @@ private:
 	* Creates the subset of tags necessary to record allocations during GMalloc and FName construction
 	*/
 	void BootstrapTagDatas();
+	void BootstrapAllocationGroups();
 	void InitializeTagDatas_SetLLMTagNames();
 	void InitializeTagDatas_FinishRegister();
 	void InitializeTagDatas();
 	void ClearTagDatas();
+	void ClearAllocationGroups();
+	void OnActiveSetsInitialized(const FTagData* DefaultTags[static_cast<uint8>(ELLMTagSet::Max)]);
 	void RegisterTagDeclaration(FLLMTagDeclaration& TagDeclaration);
-	FTagData& RegisterTagData(FName Name, FName DisplayName, FName ParentName, FName StatName, FName SummaryStatName, bool bHasEnumTag, ELLMTag EnumTag, bool bIsStatTag, ETagReferenceSource ReferenceSource, ELLMTagSet TagSet = ELLMTagSet::None);
+	FTagData& RegisterTagData(FName Name, FName DisplayName, FName ParentName, FName StatName, FName SummaryStatName,
+		bool bHasEnumTag, ELLMTag EnumTag, bool bIsStatTag, ETagReferenceSource ReferenceSource,
+		ELLMTagSet TagSet = ELLMTagSet::None);
 	/** Construct if not yet done the data on the given FTagData that relies on the presence of other TagDatas */
 	void FinishConstruct(FTagData* TagData, ETagReferenceSource ReferenceSource);
 	void ReportDuplicateTagName(FTagData* TagData, ETagReferenceSource ReferenceSource);
@@ -639,6 +646,9 @@ private:
 		ETagReferenceSource ReferenceSource = ETagReferenceSource::FunctionAPI);
 	const FTagData* FindTagData(FName Name, ELLMTagSet TagSet,
 		ETagReferenceSource ReferenceSource = ETagReferenceSource::FunctionAPI);
+	const FTagData* FindOrAddDefaultTagData(ELLMTagSet TagSet);
+	FAllocationGroup* FindOrAddAllocationGroup(const UE::LLMPrivate::FActiveTags& ActiveTags);
+	FAllocationGroup* FindAllocationGroupForCompressedAllocInfo(ELLMTag Tag);
 
 	friend class ::FLLMPauseScope;
 	friend class ::FLLMScope;
@@ -660,12 +670,27 @@ private:
 	FTagDataNameMap* TagDataNameMap;
 	/** Array to Map from ELLMTag to the TagData for that tag, used to handle LLM_SCOPE with ELLMTag */
 	FTagData** TagDataEnumMap;
+	/** All AllocationGroups that have been constructed, indexed by AllocationGroup->GetIndex() */
+	FAllocationGroupArray* AllocationGroups;
+	/** AllocationGroups that were pruned and reset, used as an allocation pool. */
+	FAllocationGroupArray* AllocationGroupsFreeList;
+	/** AllocationGroups that have 0 refcount and should be considered for pruning. */
+	FAllocationGroupSet* AllocationGroupsUnreferencedSet;
+	/** Mapping from TagSetValue to AllocationGroup, only contains active AllocationGroups. */
+	FMapActiveTagsToAllocationGroup* ActiveTagsToAllocationGroup;
 
 	UE::LLMPrivate::FTagFilter* TraceTagFilter;
 
 	FLLMTracker* Trackers[static_cast<int32>(ELLMTracker::Max)];
 
 	mutable UE::FPlatformSharedMutex TagDataMutex;
+	/**
+	 * Guards ActiveTagsToAllocationGroup, AllocationGroupsFreeList, and AllocationGroupsUnreferencedSet. Potentially
+	 * read-locked or write-locked on every allocation; threads cache the read of the containers to avoid contention.
+	 **/
+	mutable UE::FPlatformSharedMutex AllocationGroupsMutex;
+	/** Guards AllocationGroups. Read-locked on every call to Free, so we avoid write-locking it whenever possible .*/
+	mutable UE::FPlatformSharedMutex AllocationGroupsIndexMutex;
 	UE::FPlatformRecursiveMutex UpdateMutex;
 	UE::FPlatformRecursiveMutex BootstrapMutex;
 
diff --git a/Engine/Source/Runtime/Core/Private/HAL/LLM/LLMPrivateDefines.h b/Engine/Source/Runtime/Core/Private/HAL/LLM/LLMPrivateDefines.h
--- a/Engine/Source/Runtime/Core/Private/HAL/LLM/LLMPrivateDefines.h
+++ b/Engine/Source/Runtime/Core/Private/HAL/LLM/LLMPrivateDefines.h
@@ -7,18 +7,15 @@
 #if ENABLE_LOW_LEVEL_MEM_TRACKER
 
 // Private defines configuring LLM; see also public defines in LowLevelMemTracker.h
-// LLM_ALLOW_NAMES_TAGS: Set to 1 to allow arbitrary FNames to be used as tags, at the cost of more LLM memory usage per allocation. Set to 0 to store only the top-level ELLMTag for each allocation.
-// LLM_SCOPES always use FNames for the definition of their tag
-// Storing these tags on each allocation however requires 4 bytes per allocation
-// To reduce the memory overhead of LLM, this can be reduced to 1 byte per allocation at the cost of showing only the top-level 256 Tags; tags are replaced with their containing toplevel tag during allocation.
-// This setting is ignored if features requiring more information per allocation (LLM_ALLOW_STATS, LLM_ALLOW_ASSETS_TAGS) are enabled
+
+// LLM_ALLOW_NAMES_TAGS: Set to 0 to reduce memory per allocation from 20 bytes to 17 bytes, at the cost of a reduced
+// set of tags. Only ELLMTags will be displayed; all Name-based tags will be merged into their parent ELLMTag; most
+// of them will be merged into the catch-all CustomName tag. LLMTagSets and tags created by the Stat system are not
+// available when this setting is disabled.
 #ifndef LLM_ALLOW_NAMES_TAGS
 #define LLM_ALLOW_NAMES_TAGS 1
 #endif
 
-// Enable storing the full tag for each alloc if (1) arbitrary names tags are allowed or (2) Stats are allowed (stats use names tags) or (3) Asset tags are allowed (assets use names tags)
-#define LLM_ENABLED_FULL_TAGS (LLM_ALLOW_NAMES_TAGS || LLM_ALLOW_STATS || LLM_ALLOW_ASSETS_TAGS)
-
 // Whether to enable running with reduced threads. This is currently enabled because the engine crashes with -norenderthread
 #define LLM_ENABLED_REDUCE_THREADS 0
 
@@ -44,16 +41,22 @@
 
 #define LLM_CSV_PROFILER_WRITER_ENABLED CSV_PROFILER_STATS
 
-static_assert(static_cast<uint32>(ELLMTagSet::Max) == 4, "You added a new TagSet, please update the #ifs below");
-#if LLM_ALLOW_ASSETS_TAGS
-#define LLM_TAGSET_MAX 3
-#else
-#define LLM_TAGSET_MAX 2
-#endif
-
 namespace UE::LLMPrivate
 {
 
+// PruneUpdatesForTrackingData must be >= 2 for behavior correctness, see comment in
+// FLLMTracker::FetchAndClearTagSizes.
+constexpr int32 PruneUpdatesForTrackingData = 5;
+constexpr int32 PruneUpdatesForGroup = 5;
+constexpr float PrunePopulationRatioForTrackingData = 2.5f;
+
+enum class EPruningLevel : uint8
+{
+	None,
+	Default,
+	Max,
+};
+
 enum class ETagReferenceSource : uint8
 {
 	Scope,
diff --git a/Engine/Source/Runtime/Core/Public/HAL/LowLevelMemStats.h b/Engine/Source/Runtime/Core/Public/HAL/LowLevelMemStats.h
--- a/Engine/Source/Runtime/Core/Public/HAL/LowLevelMemStats.h
+++ b/Engine/Source/Runtime/Core/Public/HAL/LowLevelMemStats.h
@@ -53,9 +53,10 @@
 #endif //ENABLE_LOW_LEVEL_MEM_TRACKER && STATS
 
 /**
- * LLM Stat scope macros (these are noops if LLM is disabled or if LLM stat tags are disabled)
+ * LLM require-Stats scope macros. These are noops if LLM or STATS is defined out, and when defined in, are inactive unless
+ * toggled on by LLM command line options.
  */
-#if ENABLE_LOW_LEVEL_MEM_TRACKER && LLM_ENABLED_STAT_TAGS
+#if ENABLE_LOW_LEVEL_MEM_TRACKER && STATS
 #define LLM_SCOPED_TAG_WITH_STAT(Stat, Tracker) FLLMScope SCOPE_NAME(GET_STATFNAME(Stat), true /* bIsStatTag */, ELLMTagSet::None, Tracker);
 #define LLM_SCOPED_TAG_WITH_STAT_IN_SET(Stat, Set, Tracker) FLLMScope SCOPE_NAME(GET_STATFNAME(Stat), true /* bIsStatTag */, Set, Tracker);
 #define LLM_SCOPED_TAG_WITH_STAT_NAME(StatName, Tracker) FLLMScope SCOPE_NAME(StatName, true /* bIsStatTag */, ELLMTagSet::None, Tracker);
@@ -65,20 +66,6 @@
 #define LLM_SCOPED_SINGLE_STAT_TAG(Stat) DECLARE_LLM_MEMORY_STAT(TEXT(#Stat), Stat, STATGROUP_LLMFULL); LLM_SCOPED_TAG_WITH_STAT(Stat, ELLMTracker::Default);
 #define LLM_SCOPED_SINGLE_STAT_TAG_IN_SET(Stat, Set) DECLARE_LLM_MEMORY_STAT(TEXT(#Stat), Stat, STATGROUP_LLMFULL); LLM_SCOPED_TAG_WITH_STAT_IN_SET(Stat, Set, ELLMTracker::Default);
 #define LLM_SCOPED_PAUSE_TRACKING_WITH_STAT_AND_AMOUNT(Stat, Amount, Tracker) FLLMPauseScope SCOPE_NAME(GET_STATFNAME(Stat), true /* bIsStatTag */, Amount, Tracker, ELLMAllocType::None);
-#define LLM_SCOPE_DYNAMIC_STAT_OBJECTPATH(Object, Set) LLM_SCOPE_DYNAMIC(FName(*Object->GetPathName()), \
-					ELLMTracker::Default, Set, FLLMDynamicTagConstructorStatString(Object->GetPathName()))
-#define LLM_SCOPE_DYNAMIC_STAT_OBJECTPATH_FNAME(ObjectPath, Set) LLM_SCOPE_DYNAMIC(ObjectPath, \
-					ELLMTracker::Default, Set, FLLMDynamicTagConstructorStatString((ObjectPath).ToString()))
-
-// special stat pushing to update threads after each asset when tracking assets
-// Currently this is unused, but we may use it for optimizations later
-#if LLM_ALLOW_ASSETS_TAGS
-//#define LLM_PUSH_STATS_FOR_ASSET_TAGS() if (FLowLevelMemTracker::Get().IsTagSetActive(ELLMTagSet::Assets)) {}
-#define LLM_PUSH_STATS_FOR_ASSET_TAGS()
-#else
-#define LLM_PUSH_STATS_FOR_ASSET_TAGS()
-#endif
-
 #else
 #define LLM_SCOPED_TAG_WITH_STAT(...)
 #define LLM_SCOPED_TAG_WITH_STAT_IN_SET(...)
@@ -89,10 +76,32 @@
 #define LLM_SCOPED_SINGLE_STAT_TAG(...)
 #define LLM_SCOPED_SINGLE_STAT_TAG_IN_SET(...)
 #define LLM_SCOPED_PAUSE_TRACKING_WITH_STAT_AND_AMOUNT(...)
+#endif
+
+/**
+ * LLM optionally-use-Stats scope macros. These are active even if STATS is defined out; they still create LLM tag
+ * scopes. If STATS is defined in, they also create a STAT matching the tag (if their tagset is active in LLM). They
+ * are noops if LLM is defined out. When defined in, they are inactive unless toggled on by LLM command line options.
+ */
+#if ENABLE_LOW_LEVEL_MEM_TRACKER
+
+#define LLM_SCOPE_DYNAMIC_STAT_OBJECTPATH(Object, Set) LLM_SCOPE_DYNAMIC(FName(*Object->GetPathName()), \
+					ELLMTracker::Default, Set, FLLMDynamicTagConstructorStatString(Object->GetPathName()))
+#define LLM_SCOPE_DYNAMIC_STAT_OBJECTPATH_FNAME(ObjectPath, Set) LLM_SCOPE_DYNAMIC(ObjectPath, \
+					ELLMTracker::Default, Set, FLLMDynamicTagConstructorStatString((ObjectPath).ToString()))
+
+// Macro to update threads after each asset when tracking assets
+// Currently this is unused, but we may use it for optimizations later
+//#define LLM_PUSH_STATS_FOR_ASSET_TAGS() if (FLowLevelMemTracker::Get().IsTagSetActive(ELLMTagSet::Assets)) {}
+#define LLM_PUSH_STATS_FOR_ASSET_TAGS()
+
+#else // ENABLE_LOW_LEVEL_MEM_TRACKER
+
 #define LLM_SCOPE_DYNAMIC_STAT_OBJECTPATH(...)
 #define LLM_SCOPE_DYNAMIC_STAT_OBJECTPATH_FNAME(...)
 #define LLM_PUSH_STATS_FOR_ASSET_TAGS()
-#endif
+
+#endif // else !ENABLE_LOW_LEVEL_MEM_TRACKER
 
 #define LLM_SCOPED_TAG_WITH_OBJECT_IN_SET(Object, Set) UE_DEPRECATED_MACRO(5.3, "Use LLM_SCOPE_DYNAMIC_STAT_OBJECTPATH instead") LLM_SCOPE_DYNAMIC_STAT_OBJECTPATH(Object, Set);
 
diff --git a/Engine/Source/Runtime/Core/Public/HAL/LowLevelMemTracker.h b/Engine/Source/Runtime/Core/Public/HAL/LowLevelMemTracker.h
--- a/Engine/Source/Runtime/Core/Public/HAL/LowLevelMemTracker.h
+++ b/Engine/Source/Runtime/Core/Public/HAL/LowLevelMemTracker.h
@@ -18,36 +18,23 @@
 #endif
 #define ENABLE_LOW_LEVEL_MEM_TRACKER (LLM_ENABLED_IN_CONFIG && PLATFORM_SUPPORTS_LLM)
 
+// Deprecated defines no longer used
+#undef LLM_ALLOW_ASSETS_TAGS
+#define LLM_ALLOW_ASSETS_TAGS ENABLE_LOW_LEVEL_MEM_TRACKER
+#undef LLM_ALLOW_UOBJECTCLASSES_TAGS
+#define LLM_ALLOW_UOBJECTCLASSES_TAGS ENABLE_LOW_LEVEL_MEM_TRACKER
+#undef LLM_ALLOW_STATS
+#define LLM_ALLOW_STATS ENABLE_LOW_LEVEL_MEM_TRACKER
+#define LLM_ENABLED_STAT_TAGS ENABLE_LOW_LEVEL_MEM_TRACKER
+
+// Define LLM macros to refer to the implementation class if LLM is compiled in, otherwise define them to NOOP.
 #if ENABLE_LOW_LEVEL_MEM_TRACKER
 
-// Public defines configuring LLM; see also private defines in LowLevelMemTracker.cpp
-
-// LLM_ALLOW_ASSETS_TAGS: Set to 1 to enable run-time toggling of AssetTags reporting, 0 to disable.
-// Enabling the define causes extra cputime costs to track costs even when AssetTags are toggled off.
-// When defined on, the feature can be toggled on at runtime with commandline -llmtagsets=assets.
+// Notes about LLM CommandLine options
+// LLMTagSets: Assets: the feature can be toggled on at runtime with commandline -llmtagsets=assets.
 // Toggling the feature on causes a huge number of stat ids to be created and has a high cputime cost.
-// When defined on and runtime-toggled on, AssetTags report the asset that is in scope for each allocation
+// When toggled on, AssetTags report the asset that is in scope for each allocation
 // LLM Assets can be viewed in game using 'Stat LLMAssets'.
-#ifndef LLM_ALLOW_ASSETS_TAGS
-	#define LLM_ALLOW_ASSETS_TAGS 0
-#endif
-
-// LLM_ALLOW_UOBJECTCLASSES_TAGS will be enabled when LLM_ALLOW_ASSETS_TAGS is not
-// The intent is to avoid the amount of memory overhead by having 3 active TagSet at once.
-#ifndef LLM_ALLOW_UOBJECTCLASSES_TAGS
-	#define LLM_ALLOW_UOBJECTCLASSES_TAGS !LLM_ALLOW_ASSETS_TAGS
-#endif
-
-// LLM_ALLOW_STATS: Set to 1 to allow stats to be used as tags, 0 to disable.
-// When enabled LLM_SCOPED_TAG_WITH_STAT macros are enabled and create an LLM tag per stat at the cost of more LLM
-// memory usage per allocation. Turning this on uses the same amount of memory per allocation as LLM_ALLOW_NAMES_TAGS.
-// Turning both of them on has no extra cost.
-#ifndef LLM_ALLOW_STATS
-	#define LLM_ALLOW_STATS 0
-#endif
-
-// Enable stat tags if: (1) Stats or (2) Asset tags are allowed (asset tags use the stat macros to record asset scopes)
-#define LLM_ENABLED_STAT_TAGS (LLM_ALLOW_STATS || LLM_ALLOW_ASSETS_TAGS)
 
 #include "Containers/Array.h"
 #include "Containers/ArrayView.h"
@@ -138,7 +125,7 @@ enum class ELLMTagSet : uint8
 	AssetClasses,
 	UObjectClasses,
 	
-	Max, // note: see FLowLevelMemTracker::ShouldReduceThreads and IsAssetTagForAssets if you add any asset-style tagsets
+	Max, // note: update functions that handle ELLMTagSets if you add any new tagsets
 };
 
 namespace UE::LLM
@@ -916,14 +903,6 @@ inline bool FLowLevelMemTracker::IsEnabled()
 
 #else
 
-#ifndef LLM_ALLOW_ASSETS_TAGS
-	#define LLM_ALLOW_ASSETS_TAGS 0
-#endif
-
-#ifndef LLM_ALLOW_UOBJECTCLASSES_TAGS
-	#define LLM_ALLOW_UOBJECTCLASSES_TAGS 0
-#endif
-
 #define LLM(...)
 #define LLM_IF_ENABLED(...)
 #define LLM_IS_ENABLED() false
@@ -955,4 +934,4 @@ inline bool FLowLevelMemTracker::IsEnabled()
 #define LLM_DECLARE_BOOTSTRAP_TAG(...) 
 #define LLM_DECLARE_BOOTSTRAP_TAG_API(...)
 
-#endif		// #if ENABLE_LOW_LEVEL_MEM_TRACKER
+#endif		// #if ENABLE_LOW_LEVEL_MEM_TRACKER
\ No newline at end of file
