Do not reserve physical\swap memory for virtual memory allocation on Android, Unix and Apple platforms

#Synced-CL 49805908

diff --git a/Engine/Source/Runtime/Core/Private/Android/AndroidPlatformMemory.cpp b/Engine/Source/Runtime/Core/Private/Android/AndroidPlatformMemory.cpp
--- a/Engine/Source/Runtime/Core/Private/Android/AndroidPlatformMemory.cpp
+++ b/Engine/Source/Runtime/Core/Private/Android/AndroidPlatformMemory.cpp
@@ -397,7 +397,7 @@ uint64 FAndroidPlatformMemory::GetMemoryUsedFast()
 	return VMRSS + VMSwap;
 }
 
-// Called from JNI when trim messages are recieved.
+// Called from JNI when trim messages are received.
 void FAndroidPlatformMemory::UpdateOSMemoryStatus(EOSMemoryStatusCategory OSMemoryStatusCategory, int Value)
 {
 	switch (OSMemoryStatusCategory)
@@ -603,7 +603,7 @@ FAndroidPlatformMemory::FPlatformVirtualMemoryBlock FAndroidPlatformMemory::FPla
 	size_t Alignment = FMath::Max(InAlignment, GetVirtualSizeAlignment());
 	check(Alignment <= GetVirtualSizeAlignment());
 
-	Result.Ptr = mmap(nullptr, Result.GetActualSize(), PROT_READ | PROT_WRITE, MAP_PRIVATE | MAP_ANON, -1, 0);
+	Result.Ptr = mmap(nullptr, Result.GetActualSize(), PROT_READ | PROT_WRITE, MAP_PRIVATE | MAP_ANON | MAP_NORESERVE, -1, 0);
 	if (!LIKELY(Result.Ptr != MAP_FAILED))
 	{
 		FPlatformMemory::OnOutOfMemory(Result.GetActualSize(), InAlignment);
diff --git a/Engine/Source/Runtime/Core/Private/Apple/ApplePlatformMemory.cpp b/Engine/Source/Runtime/Core/Private/Apple/ApplePlatformMemory.cpp
--- a/Engine/Source/Runtime/Core/Private/Apple/ApplePlatformMemory.cpp
+++ b/Engine/Source/Runtime/Core/Private/Apple/ApplePlatformMemory.cpp
@@ -697,7 +697,7 @@ FApplePlatformMemory::FPlatformVirtualMemoryBlock FApplePlatformMemory::FPlatfor
 	size_t Alignment = FMath::Max(InAlignment, GetVirtualSizeAlignment());
 	check(Alignment <= GetVirtualSizeAlignment());
 
-	Result.Ptr = MmapWithTag(nullptr, Result.GetActualSize(), PROT_READ | PROT_WRITE, MAP_PRIVATE | MAP_ANON, MmapTag::AllocateVirtualMemoryBlock, 0);
+	Result.Ptr = MmapWithTag(nullptr, Result.GetActualSize(), PROT_READ | PROT_WRITE, MAP_PRIVATE | MAP_ANON | MAP_NORESERVE, MmapTag::AllocateVirtualMemoryBlock, 0);
 	if (!LIKELY(Result.Ptr != MAP_FAILED))
 	{
 		FPlatformMemory::OnOutOfMemory(Result.GetActualSize(), InAlignment);
diff --git a/Engine/Source/Runtime/Core/Private/Unix/UnixPlatformMemory.cpp b/Engine/Source/Runtime/Core/Private/Unix/UnixPlatformMemory.cpp
--- a/Engine/Source/Runtime/Core/Private/Unix/UnixPlatformMemory.cpp
+++ b/Engine/Source/Runtime/Core/Private/Unix/UnixPlatformMemory.cpp
@@ -625,7 +625,7 @@ FUnixPlatformMemory::FPlatformVirtualMemoryBlock FUnixPlatformMemory::FPlatformV
 	size_t Alignment = FMath::Max(InAlignment, GetVirtualSizeAlignment());
 	check(Alignment <= GetVirtualSizeAlignment());
 
-	Result.Ptr = mmap(nullptr, Result.GetActualSize(), PROT_READ | PROT_WRITE, MAP_PRIVATE | MAP_ANON, -1, 0);
+	Result.Ptr = mmap(nullptr, Result.GetActualSize(), PROT_READ | PROT_WRITE, MAP_PRIVATE | MAP_ANON | MAP_NORESERVE, -1, 0);
 	if (LIKELY(Result.Ptr != MAP_FAILED))
 	{
 		MarkMappedMemoryMergable(Result.Ptr, Result.GetActualSize());
