LLM: Fix correctness bug in FindTagByName and performance bug in LLMGetTagUniqueName, found by static analysis tools.
#rnx

#Synced-CL 49797803

#Synced-CL 49797803

diff --git a/Engine/Source/Runtime/Core/Private/HAL/LLM/LLM.cpp b/Engine/Source/Runtime/Core/Private/HAL/LLM/LLM.cpp
--- a/Engine/Source/Runtime/Core/Private/HAL/LLM/LLM.cpp
+++ b/Engine/Source/Runtime/Core/Private/HAL/LLM/LLM.cpp
@@ -1377,6 +1377,7 @@ FName LLMGetTagUniqueName(ELLMTag Tag)
 				UniqueNameBuffer.Appendf(TEXT("ELLMTag%d"), CreateIndex);
 				CustomNames[CreateIndex - LLM_CUSTOM_TAG_START] = FName(UniqueNameBuffer);
 			}
+			bCustomNamesInitialized = true;
 		}
 		return CustomNames[Index - LLM_CUSTOM_TAG_START];
 	}
@@ -3334,7 +3335,7 @@ bool FLLMGlobals::FindTagByNameInner(const TCHAR* Name, uint64& OutTag, ELLMTagS
 			{
 				continue;
 			}
-			if (TCString<TCHAR>::Stricmp(*TagData->GetDisplayName().ToString(), Name))
+			if (0 == TCString<TCHAR>::Stricmp(*TagData->GetDisplayName().ToString(), Name))
 			{
 				OutTag = static_cast<uint64>(TagData->GetContainingEnum());
 				return true;
