Core: Moved alignment for FIoHash and FBlake3Hash from a member to the struct

This gives static analysis tools a better understanding of the alignment and fixes false positives when casting the type or its byte array to a 32-bit integer pointer.

#rnx
#Synced-CL 49572071

diff --git a/Engine/Source/Runtime/Core/Public/Hash/Blake3.h b/Engine/Source/Runtime/Core/Public/Hash/Blake3.h
--- a/Engine/Source/Runtime/Core/Public/Hash/Blake3.h
+++ b/Engine/Source/Runtime/Core/Public/Hash/Blake3.h
@@ -17,7 +17,7 @@ class FArchive;
 ///////////////////////////////////////////////////////////////////////////////////////////////////////////////////////
 
 /** Stores a BLAKE3 hash. */
-struct FBlake3Hash
+struct alignas(uint32) FBlake3Hash
 {
 public:
 	using ByteArray = uint8[32];
@@ -72,7 +72,7 @@ public:
 	}
 
 private:
-	alignas(uint32) ByteArray Hash{};
+	ByteArray Hash{};
 };
 
 inline const FBlake3Hash FBlake3Hash::Zero;
diff --git a/Engine/Source/Runtime/Core/Public/IO/IoHash.h b/Engine/Source/Runtime/Core/Public/IO/IoHash.h
--- a/Engine/Source/Runtime/Core/Public/IO/IoHash.h
+++ b/Engine/Source/Runtime/Core/Public/IO/IoHash.h
@@ -24,7 +24,7 @@ class FArchive;
  * When the data to hash is not in a contiguous region of memory, use FIoHashBuilder to hash data
  * in blocks with FIoHashBuilder::Update(Block) followed by FIoHashBuilder::Finalize().
  */
-struct FIoHash
+struct alignas(uint32) FIoHash
 {
 public:
 	using ByteArray = uint8[20];
@@ -65,7 +65,7 @@ public:
 	static const FIoHash Zero;
 
 private:
-	alignas(uint32) ByteArray Hash{};
+	ByteArray Hash{};
 
 	friend inline bool operator==(const FIoHash& A, const FIoHash& B)
 	{
