Color: Convert transfer function constants to constexpr to guarantee compile-time evaluation.

#Synced-CL 49805348

diff --git a/Engine/Source/Runtime/Core/Public/ColorManagement/TransferFunctions.h b/Engine/Source/Runtime/Core/Public/ColorManagement/TransferFunctions.h
--- a/Engine/Source/Runtime/Core/Public/ColorManagement/TransferFunctions.h
+++ b/Engine/Source/Runtime/Core/Public/ColorManagement/TransferFunctions.h
@@ -57,12 +57,12 @@ inline float DecodeSRGB(float Value)
 */
 inline float EncodeST2084(float Value)
 {
-	const float Lp = 10000.0f;
-	const float m1 = 2610 / 4096.0f * (1.0f / 4.0f);
-	const float m2 = 2523 / 4096.0f * 128.0f;
-	const float c1 = 3424 / 4096.0f;
-	const float c2 = 2413 / 4096.0f * 32.f;
-	const float c3 = 2392 / 4096.0f * 32.f;
+	constexpr float Lp = 10000.0f;
+	constexpr float m1 = 2610 / 4096.0f * (1.0f / 4.0f);
+	constexpr float m2 = 2523 / 4096.0f * 128.0f;
+	constexpr float c1 = 3424 / 4096.0f;
+	constexpr float c2 = 2413 / 4096.0f * 32.f;
+	constexpr float c3 = 2392 / 4096.0f * 32.f;
 
 	Value = FGenericPlatformMath::Pow(Value / Lp, m1);
 	return FGenericPlatformMath::Pow((c1 + c2 * Value) / (c3 * Value + 1), m2);
@@ -75,12 +75,12 @@ inline float EncodeST2084(float Value)
 */
 inline float DecodeST2084(float Value)
 {
-	const float Lp = 10000.0f;
-	const float m1 = 2610 / 4096.0f * (1.0f / 4.0f);
-	const float m2 = 2523 / 4096.0f * 128.0f;
-	const float c1 = 3424 / 4096.0f;
-	const float c2 = 2413 / 4096.0f * 32.f;
-	const float c3 = 2392 / 4096.0f * 32.f;
+	constexpr float Lp = 10000.0f;
+	constexpr float m1 = 2610 / 4096.0f * (1.0f / 4.0f);
+	constexpr float m2 = 2523 / 4096.0f * 128.0f;
+	constexpr float c1 = 3424 / 4096.0f;
+	constexpr float c2 = 2413 / 4096.0f * 32.f;
+	constexpr float c3 = 2392 / 4096.0f * 32.f;
 
 	const float Vp = FGenericPlatformMath::Pow(Value, 1.0f / m2);
 	Value = FGenericPlatformMath::Max(0.0f, Vp - c1);
@@ -115,10 +115,10 @@ inline float DecodeGamma22(float Value)
 */
 inline float EncodeBT1886(float Value)
 {
-	const float L_B = 0;
-	const float L_W = 1;
-	float Gamma = 2.40f;
-	float GammaInv = 1.0f / Gamma;
+	constexpr float L_B = 0;
+	constexpr float L_W = 1;
+	constexpr float Gamma = 2.40f;
+	constexpr float GammaInv = 1.0f / Gamma;
 	float N = FGenericPlatformMath::Pow(L_W, GammaInv) - FGenericPlatformMath::Pow(L_B, GammaInv);
 	float A = FGenericPlatformMath::Pow(N, Gamma);
 	float B = FGenericPlatformMath::Pow(L_B, GammaInv) / N;
@@ -153,10 +153,10 @@ inline float DecodeGamma26(float Value)
 */
 inline float DecodeBT1886(float Value)
 {
-	const float L_B = 0;
-	const float L_W = 1;
-	float Gamma = 2.40f;
-	float GammaInv = 1.0f / Gamma;
+	constexpr float L_B = 0;
+	constexpr float L_W = 1;
+	constexpr float Gamma = 2.40f;
+	constexpr float GammaInv = 1.0f / Gamma;
 	float N = FGenericPlatformMath::Pow(L_W, GammaInv) - FGenericPlatformMath::Pow(L_B, GammaInv);
 	float A = FGenericPlatformMath::Pow(N, Gamma);
 	float B = FGenericPlatformMath::Pow(L_B, GammaInv) / N;
@@ -214,10 +214,10 @@ inline float DecodeREDLog(float Value)
 */
 inline float EncodeREDLog3G10(float Value)
 {
-	const float A = 0.224282f;
-	const float B = 155.975327f;
-	const float C = 0.01f;
-	const float G = 15.1927f;
+	constexpr float A = 0.224282f;
+	constexpr float B = 155.975327f;
+	constexpr float C = 0.01f;
+	constexpr float G = 15.1927f;
 
 	Value += C;
 
@@ -238,10 +238,10 @@ inline float EncodeREDLog3G10(float Value)
 */
 inline float DecodeREDLog3G10(float Value)
 {
-	const float A = 0.224282f;
-	const float B = 155.975327f;
-	const float C = 0.01f;
-	const float G = 15.1927f;
+	constexpr float A = 0.224282f;
+	constexpr float B = 155.975327f;
+	constexpr float C = 0.01f;
+	constexpr float G = 15.1927f;
 
 	if (Value < 0.0f)
 	{
@@ -354,13 +354,13 @@ inline float DecodeSLog3(float Value)
 */
 inline float EncodeArriAlexaV3LogC(float Value)
 {
-	const float cut = 0.010591f;
-	const float a = 5.555556f;
-	const float b = 0.052272f;
-	const float c = 0.247190f;
-	const float d = 0.385537f;
-	const float e = 5.367655f;
-	const float f = 0.092809f;
+	constexpr float cut = 0.010591f;
+	constexpr float a = 5.555556f;
+	constexpr float b = 0.052272f;
+	constexpr float c = 0.247190f;
+	constexpr float d = 0.385537f;
+	constexpr float e = 5.367655f;
+	constexpr float f = 0.092809f;
 
 	if (Value > cut)
 	{
@@ -379,13 +379,13 @@ inline float EncodeArriAlexaV3LogC(float Value)
 */
 inline float DecodeArriAlexaV3LogC(float Value)
 {
-	const float cut = 0.010591f;
-	const float a = 5.555556f;
-	const float b = 0.052272f;
-	const float c = 0.247190f;
-	const float d = 0.385537f;
-	const float e = 5.367655f;
-	const float f = 0.092809f;
+	constexpr float cut = 0.010591f;
+	constexpr float a = 5.555556f;
+	constexpr float b = 0.052272f;
+	constexpr float c = 0.247190f;
+	constexpr float d = 0.385537f;
+	constexpr float e = 5.367655f;
+	constexpr float f = 0.092809f;
 	
 	if (Value > e * cut + f)
 	{
@@ -459,9 +459,9 @@ inline float DecodeGoProProTune(float Value)
 */
 inline float EncodePanasonicVLog(float Value)
 {
-	const float b = 0.00873f;
-	const float c = 0.241514f;
-	const float d = 0.598206f;
+	constexpr float b = 0.00873f;
+	constexpr float c = 0.241514f;
+	constexpr float d = 0.598206f;
 
 	if( Value < 0.01f)
 	{
@@ -480,9 +480,9 @@ inline float EncodePanasonicVLog(float Value)
 */
 inline float DecodePanasonicVLog(float Value)
 {
-	const float b = 0.00873f;
-	const float c = 0.241514f;
-	const float d = 0.598206f;
+	constexpr float b = 0.00873f;
+	constexpr float c = 0.241514f;
+	constexpr float d = 0.598206f;
 
 	if (Value < 0.181f)
 	{
