making both FApplePlatformMemory and FMacPlatformMemory consistently use 64-bit structures and phys_footprint instead of resident_size when reporting UsedPhysical memory

#Synced-CL 49810091

diff --git a/Engine/Source/Runtime/Core/Private/Apple/ApplePlatformMemory.cpp b/Engine/Source/Runtime/Core/Private/Apple/ApplePlatformMemory.cpp
--- a/Engine/Source/Runtime/Core/Private/Apple/ApplePlatformMemory.cpp
+++ b/Engine/Source/Runtime/Core/Private/Apple/ApplePlatformMemory.cpp
@@ -294,9 +294,9 @@ FMalloc* FApplePlatformMemory::BaseAllocator()
         case EMemoryAllocatorToUse::Binned:
         {
             // get free memory
-            vm_statistics Stats;
+			vm_statistics64 Stats{};
             mach_msg_type_number_t StatsSize = sizeof(Stats);
-            host_statistics(mach_host_self(), HOST_VM_INFO, (host_info_t)&Stats, &StatsSize);
+            host_statistics64(mach_host_self(), HOST_VM_INFO64, (host_info_t)&Stats, &StatsSize);
             // 1 << FMath::CeilLogTwo(MemoryConstants.TotalPhysical) should really be FMath::RoundUpToPowerOfTwo,
             // but that overflows to 0 when MemoryConstants.TotalPhysical is close to 4GB, since CeilLogTwo returns 32
             // this then causes the MemoryLimit to be 0 and crashing the app
@@ -314,14 +314,20 @@ FPlatformMemoryStats FApplePlatformMemory::GetStats()
 	const FPlatformMemoryConstants& MemoryConstants = FPlatformMemory::GetConstants();
 	static FPlatformMemoryStats MemoryStats;
 	
-	// Gather platform memory stats.
-	vm_statistics Stats;
+	// Gather platform memory stats for the host and the specific process.
+	vm_statistics64 Stats{};
+	task_vm_info ProcessStats{};
 	mach_msg_type_number_t StatsSize = sizeof(Stats);
+	mach_msg_type_number_t ProcessStatsSize = sizeof(ProcessStats);
 	FMemory::Memset(&Stats, 0, StatsSize);
-	if (host_statistics(mach_host_self(), HOST_VM_INFO, (host_info_t)&Stats, &StatsSize))
+	if (host_statistics64(mach_host_self(), HOST_VM_INFO64, (host_info_t)&Stats, &StatsSize))
 	{
 		UE_LOG(LogTemp, Warning, TEXT("Failed to fetch vm statistics"));
 	}
+	if (task_info(mach_task_self(), TASK_VM_INFO, (task_info_t)&ProcessStats, &ProcessStatsSize))
+	{
+		UE_LOG(LogTemp, Warning, TEXT("Failed to fetch task-level vm statistics"));
+	}
 
 	uint64_t FreeMem = 0;
 #if PLATFORM_IOS
@@ -352,10 +358,7 @@ FPlatformMemoryStats FApplePlatformMemory::GetStats()
 		MemoryStats.PeakUsedPhysical = MemoryStats.UsedPhysical;
 	}
 
-	uint64 VMemUsed = (Stats.active_count +
-					   Stats.wire_count) *
-					  MemoryConstants.PageSize;
-	MemoryStats.UsedVirtual = VMemUsed;
+	MemoryStats.UsedVirtual = ProcessStats.virtual_size;
 	if(MemoryStats.UsedVirtual > MemoryStats.PeakUsedVirtual)
 	{
 		MemoryStats.PeakUsedVirtual = MemoryStats.UsedVirtual;
@@ -395,15 +398,15 @@ const FPlatformMemoryConstants& FApplePlatformMemory::GetConstants()
 
 		// Calculate total and available Virtual Memory
 		mach_port_t HostPort = mach_host_self();
-		mach_msg_type_number_t HostSize = sizeof(vm_statistics_data_t) / sizeof(integer_t);
+		mach_msg_type_number_t HostSize = sizeof(vm_statistics64_data_t) / sizeof(integer_t);
 
 		// verify that actual device pagesize matches defined size in vm_page_size.h
 		vm_size_t PageSize = 0;
 		host_page_size(HostPort, &PageSize);
 		ensure(vm_page_size == PageSize);
 
-		vm_statistics_data_t vm_stat;
-		if (host_statistics(HostPort, HOST_VM_INFO, (host_info_t)&vm_stat, &HostSize) != KERN_SUCCESS)
+		vm_statistics64_data_t vm_stat;
+		if (host_statistics64(HostPort, HOST_VM_INFO64, (host_info_t)&vm_stat, &HostSize) != KERN_SUCCESS)
 		{
 			UE_LOG(LogTemp, Warning, TEXT("Failed to fetch vm statistics"));
 			return MemoryConstants;
diff --git a/Engine/Source/Runtime/Core/Private/Mac/MacPlatformMemory.cpp b/Engine/Source/Runtime/Core/Private/Mac/MacPlatformMemory.cpp
--- a/Engine/Source/Runtime/Core/Private/Mac/MacPlatformMemory.cpp
+++ b/Engine/Source/Runtime/Core/Private/Mac/MacPlatformMemory.cpp
@@ -173,12 +173,16 @@ FPlatformMemoryStats FMacPlatformMemory::GetStats()
 {
 	const FPlatformMemoryConstants& MemoryConstants = FPlatformMemory::GetConstants();
 
-	static FPlatformMemoryStats MemoryStats;
+	static FPlatformMemoryStats MemoryStats{};
 
 	// Gather platform memory stats.
-	vm_statistics Stats;
+	vm_statistics64 Stats;
 	mach_msg_type_number_t StatsSize = sizeof(Stats);
-	host_statistics(mach_host_self(), HOST_VM_INFO, (host_info_t)&Stats, &StatsSize);
+	if (host_statistics64(mach_host_self(), HOST_VM_INFO64, (host_info_t)&Stats, &StatsSize) != KERN_SUCCESS)
+	{
+		UE_LOG(LogTemp, Warning, TEXT("Failed to fetch vm statistics"));
+		return MemoryStats;
+	}
 	uint64_t FreeMem = (Stats.free_count + Stats.inactive_count) * MemoryConstants.PageSize;
 	MemoryStats.AvailablePhysical = FreeMem;
 	
@@ -189,10 +193,18 @@ FPlatformMemoryStats FMacPlatformMemory::GetStats()
 	MemoryStats.AvailableVirtual = FreeMem + SwapUsage.xsu_avail;
 
 	// Just get memory information for the process and report the working set instead
-	mach_task_basic_info_data_t TaskInfo;
-	mach_msg_type_number_t TaskInfoCount = MACH_TASK_BASIC_INFO_COUNT;
-	task_info( mach_task_self(), MACH_TASK_BASIC_INFO, (task_info_t)&TaskInfo, &TaskInfoCount );
-	MemoryStats.UsedPhysical = TaskInfo.resident_size;
+	task_vm_info TaskInfo;
+	mach_msg_type_number_t TaskInfoCount = sizeof(TaskInfo);
+	if (task_info( mach_task_self(), TASK_VM_INFO, (task_info_t)&TaskInfo, &TaskInfoCount ) != KERN_SUCCESS)
+	{
+		UE_LOG(LogTemp, Warning, TEXT("Failed to fetch vm statistics"));
+		return MemoryStats;
+	}
+	
+	// we want to use phys_footprint instead of resident size as phys_footprint 
+	// represents our application's total memory footprint rather than the memory
+	// that's just in RAM right now (which can vary drastically depending on memory pressure)
+	MemoryStats.UsedPhysical = TaskInfo.phys_footprint;
 	if(MemoryStats.UsedPhysical > MemoryStats.PeakUsedPhysical)
 	{
 		MemoryStats.PeakUsedPhysical = MemoryStats.UsedPhysical;
