Add LLM scope for merged modules tracking

#rnx

#Synced-CL 49833605

diff --git a/Engine/Source/Runtime/Core/Private/Modules/ModuleManager.cpp b/Engine/Source/Runtime/Core/Private/Modules/ModuleManager.cpp
--- a/Engine/Source/Runtime/Core/Private/Modules/ModuleManager.cpp
+++ b/Engine/Source/Runtime/Core/Private/Modules/ModuleManager.cpp
@@ -545,6 +545,8 @@ void* FModuleManager::InternalLoadLibrary(FName ModuleName, const FString& Libra
 {
 	UE_LOG(LogModuleManager, Verbose, TEXT("InternalLoadLibrary: '%s' ('%s')"), *ModuleName.ToString(), *LibraryToLoad);
 
+	LLM_SCOPE_BYNAME(TEXT("Modules"));
+
 	void* Handle = nullptr;
 
 #if CPUPROFILERTRACE_ENABLED
