Made TScriptDelegate's member functions private and made a (protected) member function for setting them.

diff --git a/Engine/Source/Runtime/Core/Public/Delegates/DelegateSignatureImpl.inl b/Engine/Source/Runtime/Core/Public/Delegates/DelegateSignatureImpl.inl
--- a/Engine/Source/Runtime/Core/Public/Delegates/DelegateSignatureImpl.inl
+++ b/Engine/Source/Runtime/Core/Public/Delegates/DelegateSignatureImpl.inl
@@ -1200,10 +1200,9 @@ private:
 
 		// NOTE: If you hit a compile error on the following line, it means you're trying to use a non-UObject type
 		//       with this delegate, which is not supported
-		this->Object = const_cast<UObject*>(Cast<const UObject>(InUserObject));
+		const UObject* UserObject = Cast<UObject>(InUserObject);
 
-		// Store the function name.  The incoming function name was generated by a macro and includes the method's class name.
-		this->FunctionName = InFunctionName;
+		this->BindProtected(UserObject, InFunctionName);
 
 		ensureMsgf(this->IsBound(), TEXT("Unable to bind delegate to '%s' (function might not be marked as a UFUNCTION or object may be pending kill)"), *InFunctionName.ToString());
 	}
diff --git a/Engine/Source/Runtime/Core/Public/UObject/ScriptDelegates.h b/Engine/Source/Runtime/Core/Public/UObject/ScriptDelegates.h
--- a/Engine/Source/Runtime/Core/Public/UObject/ScriptDelegates.h
+++ b/Engine/Source/Runtime/Core/Public/UObject/ScriptDelegates.h
@@ -539,16 +539,19 @@ public:
 	}
 
 protected:
+	void BindProtected(WeakPtrType InObject, FName InFunctionName)
+	{
+		this->Object = InObject;
+		this->FunctionName = InFunctionName;
+	}
 
+private:
 	/** The object bound to this delegate, or nullptr if no object is bound */
 	WeakPtrType Object;
 
 	/** Name of the function to call on the bound object */
 	FName FunctionName;
 
-	// 
-	friend class FCallDelegateHelper;
-
 	friend struct TIsZeroConstructType<TScriptDelegate>;
 };
 
