[Core] Add cvar to disable debugger breaking when ensure statements are hit
* Rename ExecCheckImplInternal to CheckEnsureFailed to be more clear that it applies to ensures and not checks

#Synced-CL 49832859

diff --git a/Engine/Source/Runtime/Core/Private/Misc/AssertionMacros.cpp b/Engine/Source/Runtime/Core/Private/Misc/AssertionMacros.cpp
--- a/Engine/Source/Runtime/Core/Private/Misc/AssertionMacros.cpp
+++ b/Engine/Source/Runtime/Core/Private/Misc/AssertionMacros.cpp
@@ -139,6 +139,13 @@ FAutoConsoleVariableRef CVarEnsureAlwaysEnabled(
 	TEXT("Set to false to turn ensureAlways into regular ensure"),
 	ECVF_Default);
 
+bool GEnsureBreakEnabled = true;
+FAutoConsoleVariableRef CVarEnsureBreakEnabled(
+	TEXT("core.EnsureBreakEnabled"),
+	GEnsureBreakEnabled,
+	TEXT("Set to false to turn off the debugger catching ensures"),
+	ECVF_Default);
+
 int32 GEnsureResetState = 1;
 FAutoConsoleCommand CVarResetEnsureState(
 	TEXT("core.ResetEnsureState"),
@@ -925,7 +932,7 @@ bool UE_COLD UE_DEBUG_SECTION VARARGS CheckVerifyImpl(const ANSICHAR* File, int3
 }
 
 UE_AUTORTFM_ALWAYS_OPEN
-bool UE_COLD UE_DEBUG_SECTION UE::Assert::Private::ExecCheckImplInternal(std::atomic<uint8>& bExecuted, bool bAlways, const ANSICHAR* File, int32 Line, const ANSICHAR* Expr)
+bool UE_COLD UE_DEBUG_SECTION UE::Assert::Private::CheckEnsureFailed(bool bAlways, std::atomic<uint8>& bExecuted, const ANSICHAR* File, int32 Line, const ANSICHAR* Expr)
 {
 	if (((bAlways && GEnsureAlwaysEnabled) || bExecuted.load(std::memory_order_relaxed) != GEnsureResetState) && FPlatformMisc::IsEnsureAllowed())
 	{
@@ -936,7 +943,7 @@ bool UE_COLD UE_DEBUG_SECTION UE::Assert::Private::ExecCheckImplInternal(std::at
 
 		va_list Args = {};
 		const bool bResult = CheckVerifyImpl(File, Line, PLATFORM_RETURN_ADDRESS(), Expr, TEXT(""), Args);
-		return bResult;
+		return bResult && GEnsureBreakEnabled;
 	}
 
 	return false;
@@ -957,7 +964,7 @@ bool UE_COLD UE_DEBUG_SECTION EnsureFailedV(std::atomic<uint8>& bExecuted, const
 	}
 
 	const bool bResult = CheckVerifyImpl(Ensure->File, Ensure->Line, ProgramCounter, Ensure->Expression, Ensure->Format, Args);
-	return bResult;
+	return bResult && GEnsureBreakEnabled;
 }
 
 bool UE_COLD UE_DEBUG_SECTION VARARGS UE::Assert::Private::EnsureFailed(std::atomic<uint8>& bExecuted, const FStaticEnsureRecord* Ensure, ...)
diff --git a/Engine/Source/Runtime/Core/Public/Misc/AssertionMacros.h b/Engine/Source/Runtime/Core/Public/Misc/AssertionMacros.h
--- a/Engine/Source/Runtime/Core/Public/Misc/AssertionMacros.h
+++ b/Engine/Source/Runtime/Core/Public/Misc/AssertionMacros.h
@@ -397,9 +397,9 @@ RetType UE_COLD UE_DEBUG_SECTION DispatchCheckVerify (InnerType&& Inner, ArgType
 		}
 	};
 
-	CORE_API bool UE_COLD UE_DEBUG_SECTION VARARGS CheckEnsureFailed(bool bAlways, const std::atomic<uint8>& bExecuted);
+	CORE_API bool UE_COLD UE_DEBUG_SECTION CheckEnsureFailed(bool bAlways, std::atomic<uint8>& bExecuted, const ANSICHAR* File, int32 Line, const ANSICHAR* Expr);
+	CORE_API bool UE_COLD UE_DEBUG_SECTION CheckEnsureFailed(bool bAlways, const std::atomic<uint8>& bExecuted);
 	CORE_API bool UE_COLD UE_DEBUG_SECTION VARARGS EnsureFailed(std::atomic<uint8>& bExecuted, const FStaticEnsureRecord* Ensure, ...);
-	CORE_API bool UE_COLD UE_DEBUG_SECTION ExecCheckImplInternal(std::atomic<uint8>& bExecuted, bool bAlways, const ANSICHAR* File, int32 Line, const ANSICHAR* Expr);
 
 	} // UE::Assert::Private
 
@@ -434,9 +434,9 @@ RetType UE_COLD UE_DEBUG_SECTION DispatchCheckVerify (InnerType&& Inner, ArgType
 		( \
 			LIKELY(!!(InExpression)) || \
 			( \
-				::UE::Assert::Private::ExecCheckImplInternal( \
-					::bGEnsureHasExecuted<FileLineHashForEnsure(__FILE__, __LINE__)>, \
+				::UE::Assert::Private::CheckEnsureFailed( \
 					Always, \
+					::bGEnsureHasExecuted<FileLineHashForEnsure(__FILE__, __LINE__)>, \
 					__FILE__, \
 					__LINE__, \
 					#InExpression \
