No longer perform a normal config if the PreSaveConfigFileDelegate returns that it saved the file. This can reduce some slow unnecessary disk operations on some platforms. Can revert to original behavior by either setting ini.SkipSaveIfPreSaveSucceeded to 0, or by not incrementing SavedCount from your delegate implementation
[RN] minor, foundation, core
#Synced-CL 49643405

diff --git a/Engine/Source/Runtime/Core/Private/Misc/ConfigCacheIni.cpp b/Engine/Source/Runtime/Core/Private/Misc/ConfigCacheIni.cpp
--- a/Engine/Source/Runtime/Core/Private/Misc/ConfigCacheIni.cpp
+++ b/Engine/Source/Runtime/Core/Private/Misc/ConfigCacheIni.cpp
@@ -97,6 +97,12 @@ static FAutoConsoleVariableRef CVarConfigBranchesToNeverUnload(
 	GConfigBranchesToNeverUnload,
 	TEXT("A comma separated list of config branch names that should never be unloaded."));
 
+static int GSkipSaveIfPreSaveSucceeded = 1;
+static FAutoConsoleVariableRef CVarSkipSaveIfPreSaveSucceeded(
+	TEXT("ini.SkipSaveIfPreSaveSucceeded"),
+	GSkipSaveIfPreSaveSucceeded,
+	TEXT("If true, then the standard config file save will be skipped if the PreSave delegate indicates someone saved."));
+
 
 #if WITH_EDITOR
 // editor wants full replay - and we can't put this in a cvar since we need way too early for cvar processing!
@@ -460,26 +466,28 @@ static bool SaveConfigFileWrapper(const TCHAR* IniFile, const FString& Contents)
 	int32 SavedCount = 0;
 	FCoreDelegates::TSPreSaveConfigFileDelegate().Broadcast(IniFile, Contents, SavedCount);
 
-	// save it even if a delegate did as well
+	// save if the delegate didn't, or it did and we don't want to skip the normal save
 	bool bLocalWriteSucceeded = false;
-
-	if (FConfigFile::WriteTempFileThenMove())
+	if (SavedCount == 0 || GSkipSaveIfPreSaveSucceeded == false)
 	{
-		const FString BaseFilename = FPaths::GetBaseFilename(IniFile);
-		const FString TempFilename = FPaths::CreateTempFilename(*FPaths::ProjectSavedDir(), *BaseFilename.Left(32));
-		bLocalWriteSucceeded = FFileHelper::SaveStringToFile(Contents, *TempFilename, FFileHelper::EEncodingOptions::ForceUTF8WithoutBOM);
-		if (bLocalWriteSucceeded)
+		if (FConfigFile::WriteTempFileThenMove())
 		{
-			if (!IFileManager::Get().Move(IniFile, *TempFilename))
+			const FString BaseFilename = FPaths::GetBaseFilename(IniFile);
+			const FString TempFilename = FPaths::CreateTempFilename(*FPaths::ProjectSavedDir(), *BaseFilename.Left(32));
+			bLocalWriteSucceeded = FFileHelper::SaveStringToFile(Contents, *TempFilename, FFileHelper::EEncodingOptions::ForceUTF8WithoutBOM);
+			if (bLocalWriteSucceeded)
 			{
-				IFileManager::Get().Delete(*TempFilename);
-				bLocalWriteSucceeded = false;
+				if (!IFileManager::Get().Move(IniFile, *TempFilename))
+				{
+					IFileManager::Get().Delete(*TempFilename);
+					bLocalWriteSucceeded = false;
+				}
 			}
 		}
-	}
-	else
-	{
-		bLocalWriteSucceeded = FFileHelper::SaveStringToFile(Contents, IniFile, FFileHelper::EEncodingOptions::ForceUTF8WithoutBOM);
+		else
+		{
+			bLocalWriteSucceeded = FFileHelper::SaveStringToFile(Contents, IniFile, FFileHelper::EEncodingOptions::ForceUTF8WithoutBOM);
+		}
 	}
 
 	// success is based on a delegate or file write working (or both)
