[AutoRTFM] Fix DLLEXPORT / DLLIMPORT attribute on AutoRTFMUE.h APIs

This is part of Core, not AutoRTFM.

#Synced-CL 49797767

diff --git a/Engine/Source/Runtime/Core/Public/AutoRTFM/AutoRTFMUE.h b/Engine/Source/Runtime/Core/Public/AutoRTFM/AutoRTFMUE.h
--- a/Engine/Source/Runtime/Core/Public/AutoRTFM/AutoRTFMUE.h
+++ b/Engine/Source/Runtime/Core/Public/AutoRTFM/AutoRTFMUE.h
@@ -2,13 +2,10 @@
 
 #pragma once
 
-// Include Platform.h before AutoRTFM.h so that DLLEXPORT / DLLIMPORT is defined.
-#include "HAL/Platform.h"
-
 #include "AutoRTFM.h"
 #include "Logging/LogMacros.h"
 
-UE_AUTORTFM_API DECLARE_LOG_CATEGORY_EXTERN(LogAutoRTFM, Display, All)
+CORE_API DECLARE_LOG_CATEGORY_EXTERN(LogAutoRTFM, Display, All)
 
 namespace AutoRTFM
 {
@@ -20,7 +17,7 @@ using RollbackWriteFn = void (*)(void* Address, const void* Value, size_t Size,
 // Initializes AutoRTFM for use by the Unreal Engine.
 // RollbackWrite is an optional function pointer used to rollback writes written
 // with the AutoRTFM::EWriteFlags::CustomRollback flag.
-UE_AUTORTFM_API void InitializeForUE(RollbackWriteFn RollbackWrite = nullptr);
+CORE_API void InitializeForUE(RollbackWriteFn RollbackWrite = nullptr);
 
 #else
 
