Added API support for dynamic delegates with payloads (void returns only), but still disabled due to incomplete support (notably, serialization).

diff --git a/Engine/Source/Runtime/Core/Public/Concepts/UObject.h b/Engine/Source/Runtime/Core/Public/Concepts/UObject.h
new file mode 100644
--- /dev/null
+++ b/Engine/Source/Runtime/Core/Public/Concepts/UObject.h
@@ -0,0 +1,19 @@
+// Copyright Epic Games, Inc. All Rights Reserved.
+
+#pragma once
+
+#include "Concepts/DerivedFrom.h"
+
+class UObject;
+
+namespace UE
+{
+	/**
+	 * Concept which describes a UObject, that is a type derived from UObject.
+	 * Pointers are not UObject types - use the following to match pointers to UObjects:
+	 *
+	 * UE::CPointer<Type> && UE::CSigned<std::remove_pointer_t<Type>>
+	 */
+	template <typename T>
+	concept CUObject = UE::CDerivedFrom<T, UObject>;
+}
diff --git a/Engine/Source/Runtime/Core/Public/Delegates/Delegate.h b/Engine/Source/Runtime/Core/Public/Delegates/Delegate.h
--- a/Engine/Source/Runtime/Core/Public/Delegates/Delegate.h
+++ b/Engine/Source/Runtime/Core/Public/Delegates/Delegate.h
@@ -311,7 +311,7 @@ namespace UE::Core::Private
  * @param	UserObject		UObject instance
  * @param	FuncName		Function pointer to member UFUNCTION, usually in form &UClassName::FunctionName
  */
-#define BindDynamic( UserObject, FuncName ) __Internal_BindDynamic( UserObject, FuncName, UE_PRIVATE_STATIC_FUNCTION_FNAME( #FuncName ) )
+#define BindDynamic( UserObject, FuncName, ... ) __Internal_BindDynamic( UserObject, FuncName, UE_PRIVATE_STATIC_FUNCTION_FNAME( #FuncName ), ##__VA_ARGS__ )
 
 /**
  * Helper macro to bind a UObject instance and a member UFUNCTION to a dynamic multi-cast delegate.
@@ -319,7 +319,7 @@ namespace UE::Core::Private
  * @param	UserObject		UObject instance
  * @param	FuncName		Function pointer to member UFUNCTION, usually in form &UClassName::FunctionName
  */
-#define AddDynamic( UserObject, FuncName ) __Internal_AddDynamic( UserObject, FuncName, UE_PRIVATE_STATIC_FUNCTION_FNAME( #FuncName ) )
+#define AddDynamic( UserObject, FuncName, ... ) __Internal_AddDynamic( UserObject, FuncName, UE_PRIVATE_STATIC_FUNCTION_FNAME( #FuncName ), ##__VA_ARGS__ )
 
 /**
  * Helper macro to bind a UObject instance and a member UFUNCTION to a dynamic multi-cast delegate, but only if it hasn't been bound before.
@@ -327,7 +327,7 @@ namespace UE::Core::Private
  * @param	UserObject		UObject instance
  * @param	FuncName		Function pointer to member UFUNCTION, usually in form &UClassName::FunctionName
  */
-#define AddUniqueDynamic( UserObject, FuncName ) __Internal_AddUniqueDynamic( UserObject, FuncName, UE_PRIVATE_STATIC_FUNCTION_FNAME( #FuncName ) )
+#define AddUniqueDynamic( UserObject, FuncName, ... ) __Internal_AddUniqueDynamic( UserObject, FuncName, UE_PRIVATE_STATIC_FUNCTION_FNAME( #FuncName ), ##__VA_ARGS__ )
 
 /**
  * Helper macro to unbind a UObject instance and a member UFUNCTION from this multi-cast delegate.
@@ -335,7 +335,7 @@ namespace UE::Core::Private
  * @param	UserObject		UObject instance
  * @param	FuncName		Function pointer to member UFUNCTION, usually in form &UClassName::FunctionName
  */
-#define RemoveDynamic( UserObject, FuncName ) __Internal_RemoveDynamic( UserObject, FuncName, UE_PRIVATE_STATIC_FUNCTION_FNAME( #FuncName ) )
+#define RemoveDynamic( UserObject, FuncName, ... ) __Internal_RemoveDynamic( UserObject, FuncName, UE_PRIVATE_STATIC_FUNCTION_FNAME( #FuncName ), ##__VA_ARGS__ )
 
  /**
   * Helper macro to tests if a UObject instance and a member UFUNCTION are already bound to this multi-cast delegate.
@@ -343,7 +343,7 @@ namespace UE::Core::Private
   * @param	UserObject		UObject instance
   * @param	FuncName		Function pointer to member UFUNCTION, usually in form &UClassName::FunctionName
   */
-#define IsAlreadyBound( UserObject, FuncName ) __Internal_IsAlreadyBound( UserObject, FuncName, UE_PRIVATE_STATIC_FUNCTION_FNAME( #FuncName ) )
+#define IsAlreadyBound( UserObject, FuncName, ... ) __Internal_IsAlreadyBound( UserObject, FuncName, UE_PRIVATE_STATIC_FUNCTION_FNAME( #FuncName ), ##__VA_ARGS__ )
 
 
 /*********************************************************************************************************************/
diff --git a/Engine/Source/Runtime/Core/Public/Delegates/DelegateSignatureImpl.inl b/Engine/Source/Runtime/Core/Public/Delegates/DelegateSignatureImpl.inl
--- a/Engine/Source/Runtime/Core/Public/Delegates/DelegateSignatureImpl.inl
+++ b/Engine/Source/Runtime/Core/Public/Delegates/DelegateSignatureImpl.inl
@@ -1182,27 +1182,34 @@ public:
 	/**
 	 * Templated helper class to define a typedef for user's method pointer, then used below
 	 */
-	template <bool bConst, typename UserClass>
+	template <bool bConst, typename UserClass, typename... VarTypes>
 	struct TMethodPtrResolver
 	{
+		// The decay_t means that captured payloads will need to be passed by value, including things like FString.  Will need a mapping for that.
 		using FMethodPtr = std::conditional_t<
 			bConst,
-			RetValType(UserClass::*)(ParamTypes... Params) const,
-			RetValType(UserClass::*)(ParamTypes... Params)
+			RetValType(UserClass::*)(ParamTypes... Params, std::decay_t<VarTypes>...) const,
+			RetValType(UserClass::*)(ParamTypes... Params, std::decay_t<VarTypes>...)
 		>;
 	};
 
 private:
 	template <typename UserClass, typename... VarTypes>
-	void BindDynamicImpl(const UserClass* InUserObject, FName InFunctionName)
+	void BindDynamicImpl(const UserClass* InUserObject, FName InFunctionName, VarTypes&&... Vars)
 	{
+		static_assert(UE_USE_DYNAMIC_DELEGATE_PAYLOADS || sizeof...(VarTypes) == 0, "Payload support for dynamic delegates is disabled");
+
 		check(InUserObject);
 
 		// NOTE: If you hit a compile error on the following line, it means you're trying to use a non-UObject type
 		//       with this delegate, which is not supported
 		const UObject* UserObject = Cast<UObject>(InUserObject);
 
-		this->BindProtected(UserObject, InFunctionName);
+#if UE_USE_DYNAMIC_DELEGATE_PAYLOADS
+		this->template BindProtected<RetValType(ParamTypes..., std::decay_t<VarTypes>...)>(UserObject, InFunctionName, Forward<VarTypes>(Vars)...);
+#else
+		this->BindProtected(UserObject, InFunctionName, Forward<VarTypes>(Vars)...);
+#endif
 
 		ensureMsgf(this->IsBound(), TEXT("Unable to bind delegate to '%s' (function might not be marked as a UFUNCTION or object may be pending kill)"), *InFunctionName.ToString());
 	}
@@ -1220,27 +1227,27 @@ public:
 	 *
 	 * NOTE:  We're not actually using the incoming method pointer.  We simply require it for type-safety reasons.
 	 */
-	template <typename UserClass>
-	UE_REWRITE void __Internal_BindDynamic(UserClass* InUserObject, typename TMethodPtrResolver<false, UserClass>::FMethodPtr InMethodPtr, FName InFunctionName)
+	template <typename UserClass, typename... VarTypes>
+	UE_REWRITE void __Internal_BindDynamic(UserClass* InUserObject, typename TMethodPtrResolver<false, UserClass, VarTypes...>::FMethodPtr InMethodPtr, FName InFunctionName, VarTypes&&... Vars)
 	{
 		UE_STATIC_DEPRECATE(5.8, std::is_const_v<UserClass>, "Binding a delegate with a const object pointer and non-const function is deprecated.");
-		this->BindDynamicImpl(InUserObject, InFunctionName);
+		this->BindDynamicImpl(InUserObject, InFunctionName, Forward<VarTypes>(Vars)...);
 	}
-	template <typename UserClass>
-	UE_REWRITE void __Internal_BindDynamic(TObjectPtr<UserClass> InUserObject, typename TMethodPtrResolver<false, UserClass>::FMethodPtr InMethodPtr, FName InFunctionName)
+	template <typename UserClass, typename... VarTypes>
+	UE_REWRITE void __Internal_BindDynamic(TObjectPtr<UserClass> InUserObject, typename TMethodPtrResolver<false, UserClass, VarTypes...>::FMethodPtr InMethodPtr, FName InFunctionName, VarTypes&&... Vars)
 	{
 		UE_STATIC_DEPRECATE(5.8, std::is_const_v<UserClass>, "Binding a delegate with a const object pointer and non-const function is deprecated.");
-		this->BindDynamicImpl(ToRawPtr(InUserObject), InFunctionName);
+		this->BindDynamicImpl(ToRawPtr(InUserObject), InFunctionName, Forward<VarTypes>(Vars)...);
 	}
-	template <typename UserClass>
-	UE_REWRITE void __Internal_BindDynamic(UserClass* InUserObject, typename TMethodPtrResolver<true, UserClass>::FMethodPtr InMethodPtr, FName InFunctionName)
+	template <typename UserClass, typename... VarTypes>
+	UE_REWRITE void __Internal_BindDynamic(UserClass* InUserObject, typename TMethodPtrResolver<true, UserClass, VarTypes...>::FMethodPtr InMethodPtr, FName InFunctionName, VarTypes&&... Vars)
 	{
-		this->BindDynamicImpl(InUserObject, InFunctionName);
+		this->BindDynamicImpl(InUserObject, InFunctionName, Forward<VarTypes>(Vars)...);
 	}
-	template <typename UserClass>
-	UE_REWRITE void __Internal_BindDynamic(TObjectPtr<UserClass> InUserObject, typename TMethodPtrResolver<true, UserClass>::FMethodPtr InMethodPtr, FName InFunctionName)
+	template <typename UserClass, typename... VarTypes>
+	UE_REWRITE void __Internal_BindDynamic(TObjectPtr<UserClass> InUserObject, typename TMethodPtrResolver<true, UserClass, VarTypes...>::FMethodPtr InMethodPtr, FName InFunctionName, VarTypes&&... Vars)
 	{
-		this->BindDynamicImpl(ToRawPtr(InUserObject), InFunctionName);
+		this->BindDynamicImpl(ToRawPtr(InUserObject), InFunctionName, Forward<VarTypes>(Vars)...);
 	}
 
 	/** Execute the delegate.  If the function pointer is not valid, an error will occur. */
@@ -1293,11 +1300,13 @@ public:
 	{ }
 
 private:
-	template <typename UserClass>
-	bool IsAlreadyBoundImpl(UserClass* InUserObject, FName InFunctionName) const
+	template <typename UserClass, typename... VarTypes>
+	bool IsAlreadyBoundImpl(const UserClass* InUserObject, FName InFunctionName, VarTypes&&... Vars) const
 	{
+		static_assert(UE_USE_DYNAMIC_DELEGATE_PAYLOADS || sizeof...(VarTypes) == 0, "Payload support for dynamic delegates is disabled");
+
 		check(InUserObject);
-		return this->Contains(InUserObject, InFunctionName);
+		return this->Contains(InUserObject, InFunctionName, Forward<VarTypes>(Vars)...);
 	}
 
 public:
@@ -1314,35 +1323,37 @@ public:
 	 *
 	 * NOTE:  We're not actually using the incoming method pointer.  We simply require it for type-safety reasons.
 	 */
-	template <typename UserClass>
-	UE_REWRITE bool __Internal_IsAlreadyBound(UserClass* InUserObject, typename FDelegate::template TMethodPtrResolver<false, UserClass>::FMethodPtr InMethodPtr, FName InFunctionName) const
+	template <typename UserClass, typename... VarTypes>
+	UE_REWRITE bool __Internal_IsAlreadyBound(UserClass* InUserObject, typename FDelegate::template TMethodPtrResolver<false, UserClass, VarTypes...>::FMethodPtr InMethodPtr, FName InFunctionName, VarTypes&&... Vars) const
 	{
-		return this->IsAlreadyBoundImpl(InUserObject, InFunctionName);
+		return this->IsAlreadyBoundImpl(InUserObject, InFunctionName, Forward<VarTypes>(Vars)...);
 	}
-	template <typename UserClass>
-	UE_REWRITE bool __Internal_IsAlreadyBound(TObjectPtr<UserClass> InUserObject, typename FDelegate::template TMethodPtrResolver<false, UserClass>::FMethodPtr InMethodPtr, FName InFunctionName) const
+	template <typename UserClass, typename... VarTypes>
+	UE_REWRITE bool __Internal_IsAlreadyBound(TObjectPtr<UserClass> InUserObject, typename FDelegate::template TMethodPtrResolver<false, UserClass, VarTypes...>::FMethodPtr InMethodPtr, FName InFunctionName, VarTypes&&... Vars) const
 	{
-		return this->IsAlreadyBoundImpl(ToRawPtr(InUserObject), InFunctionName);
+		return this->IsAlreadyBoundImpl(ToRawPtr(InUserObject), InFunctionName, Forward<VarTypes>(Vars)...);
 	}
-	template <typename UserClass>
-	UE_REWRITE bool __Internal_IsAlreadyBound(UserClass* InUserObject, typename FDelegate::template TMethodPtrResolver<true, UserClass>::FMethodPtr InMethodPtr, FName InFunctionName) const
+	template <typename UserClass, typename... VarTypes>
+	UE_REWRITE bool __Internal_IsAlreadyBound(UserClass* InUserObject, typename FDelegate::template TMethodPtrResolver<true, UserClass, VarTypes...>::FMethodPtr InMethodPtr, FName InFunctionName, VarTypes&&... Vars) const
 	{
-		return this->IsAlreadyBoundImpl(InUserObject, InFunctionName);
+		return this->IsAlreadyBoundImpl(InUserObject, InFunctionName, Forward<VarTypes>(Vars)...);
 	}
-	template <typename UserClass>
-	UE_REWRITE bool __Internal_IsAlreadyBound(TObjectPtr<UserClass> InUserObject, typename FDelegate::template TMethodPtrResolver<true, UserClass>::FMethodPtr InMethodPtr, FName InFunctionName) const
+	template <typename UserClass, typename... VarTypes>
+	UE_REWRITE bool __Internal_IsAlreadyBound(TObjectPtr<UserClass> InUserObject, typename FDelegate::template TMethodPtrResolver<true, UserClass, VarTypes...>::FMethodPtr InMethodPtr, FName InFunctionName, VarTypes&&... Vars) const
 	{
-		return this->IsAlreadyBoundImpl(ToRawPtr(InUserObject), InFunctionName);
+		return this->IsAlreadyBoundImpl(ToRawPtr(InUserObject), InFunctionName, Forward<VarTypes>(Vars)...);
 	}
 
 private:
-	template <typename UserClass>
-	void AddDynamicImpl(const UserClass* InUserObject, FName InFunctionName)
+	template <typename UserClass, typename... VarTypes>
+	void AddDynamicImpl(const UserClass* InUserObject, FName InFunctionName, VarTypes&&... Vars)
 	{
+		static_assert(UE_USE_DYNAMIC_DELEGATE_PAYLOADS || sizeof...(VarTypes) == 0, "Payload support for dynamic delegates is disabled");
+
 		check(InUserObject);
 
 		FDelegate NewDelegate;
-		NewDelegate.BindDynamicImpl(InUserObject, InFunctionName);
+		NewDelegate.BindDynamicImpl(InUserObject, InFunctionName, Forward<VarTypes>(Vars)...);
 
 		this->Add(NewDelegate);
 	}
@@ -1360,37 +1371,39 @@ public:
 	 *
 	 * NOTE:  We're not actually using the incoming method pointer.  We simply require it for type-safety reasons.
 	 */
-	template <typename UserClass>
-	UE_REWRITE void __Internal_AddDynamic(UserClass* InUserObject, typename FDelegate::template TMethodPtrResolver<false, UserClass>::FMethodPtr InMethodPtr, FName InFunctionName)
+	template <typename UserClass, typename... VarTypes>
+	UE_REWRITE void __Internal_AddDynamic(UserClass* InUserObject, typename FDelegate::template TMethodPtrResolver<false, UserClass, VarTypes...>::FMethodPtr InMethodPtr, FName InFunctionName, VarTypes&&... Vars)
 	{
 		UE_STATIC_DEPRECATE(5.8, std::is_const_v<UserClass>, "Binding a delegate with a const object pointer and non-const function is deprecated.");
-		this->AddDynamicImpl(InUserObject, InFunctionName);
+		this->AddDynamicImpl(InUserObject, InFunctionName, Forward<VarTypes>(Vars)...);
 	}
-	template <typename UserClass>
-	UE_REWRITE void __Internal_AddDynamic(TObjectPtr<UserClass> InUserObject, typename FDelegate::template TMethodPtrResolver<false, UserClass>::FMethodPtr InMethodPtr, FName InFunctionName)
+	template <typename UserClass, typename... VarTypes>
+	UE_REWRITE void __Internal_AddDynamic(TObjectPtr<UserClass> InUserObject, typename FDelegate::template TMethodPtrResolver<false, UserClass, VarTypes...>::FMethodPtr InMethodPtr, FName InFunctionName, VarTypes&&... Vars)
 	{
 		UE_STATIC_DEPRECATE(5.8, std::is_const_v<UserClass>, "Binding a delegate with a const object pointer and non-const function is deprecated.");
-		this->AddDynamicImpl(ToRawPtr(InUserObject), InFunctionName);
+		this->AddDynamicImpl(ToRawPtr(InUserObject), InFunctionName, Forward<VarTypes>(Vars)...);
 	}
-	template <typename UserClass>
-	UE_REWRITE void __Internal_AddDynamic(UserClass* InUserObject, typename FDelegate::template TMethodPtrResolver<true, UserClass>::FMethodPtr InMethodPtr, FName InFunctionName)
+	template <typename UserClass, typename... VarTypes>
+	UE_REWRITE void __Internal_AddDynamic(UserClass* InUserObject, typename FDelegate::template TMethodPtrResolver<true, UserClass, VarTypes...>::FMethodPtr InMethodPtr, FName InFunctionName, VarTypes&&... Vars)
 	{
-		this->AddDynamicImpl(InUserObject, InFunctionName);
+		this->AddDynamicImpl(InUserObject, InFunctionName, Forward<VarTypes>(Vars)...);
 	}
-	template <typename UserClass>
-	UE_REWRITE void __Internal_AddDynamic(TObjectPtr<UserClass> InUserObject, typename FDelegate::template TMethodPtrResolver<true, UserClass>::FMethodPtr InMethodPtr, FName InFunctionName)
+	template <typename UserClass, typename... VarTypes>
+	UE_REWRITE void __Internal_AddDynamic(TObjectPtr<UserClass> InUserObject, typename FDelegate::template TMethodPtrResolver<true, UserClass, VarTypes...>::FMethodPtr InMethodPtr, FName InFunctionName, VarTypes&&... Vars)
 	{
-		this->AddDynamicImpl(ToRawPtr(InUserObject), InFunctionName);
+		this->AddDynamicImpl(ToRawPtr(InUserObject), InFunctionName, Forward<VarTypes>(Vars)...);
 	}
 
 private:
-	template <typename UserClass>
-	void AddUniqueDynamicImpl(const UserClass* InUserObject, FName InFunctionName)
+	template <typename UserClass, typename... VarTypes>
+	void AddUniqueDynamicImpl(const UserClass* InUserObject, FName InFunctionName, VarTypes&&... Vars)
 	{
+		static_assert(UE_USE_DYNAMIC_DELEGATE_PAYLOADS || sizeof...(VarTypes) == 0, "Payload support for dynamic delegates is disabled");
+
 		check(InUserObject);
 
 		FDelegate NewDelegate;
-		NewDelegate.BindDynamicImpl(InUserObject, InFunctionName);
+		NewDelegate.BindDynamicImpl(InUserObject, InFunctionName, Forward<VarTypes>(Vars)...);
 
 		this->AddUnique(NewDelegate);
 	}
@@ -1408,36 +1421,38 @@ public:
 	 *
 	 * NOTE:  We're not actually using the incoming method pointer.  We simply require it for type-safety reasons.
 	 */
-	template <typename UserClass>
-	UE_REWRITE void __Internal_AddUniqueDynamic(UserClass* InUserObject, typename FDelegate::template TMethodPtrResolver<false, UserClass>::FMethodPtr InMethodPtr, FName InFunctionName)
+	template <typename UserClass, typename... VarTypes>
+	UE_REWRITE void __Internal_AddUniqueDynamic(UserClass* InUserObject, typename FDelegate::template TMethodPtrResolver<false, UserClass, VarTypes...>::FMethodPtr InMethodPtr, FName InFunctionName, VarTypes&&... Vars)
 	{
 		UE_STATIC_DEPRECATE(5.8, std::is_const_v<UserClass>, "Binding a delegate with a const object pointer and non-const function is deprecated.");
-		this->AddUniqueDynamicImpl(InUserObject, InFunctionName);
+		this->AddUniqueDynamicImpl(InUserObject, InFunctionName, Forward<VarTypes>(Vars)...);
 	}
-	template <typename UserClass>
-	UE_REWRITE void __Internal_AddUniqueDynamic(TObjectPtr<UserClass> InUserObject, typename FDelegate::template TMethodPtrResolver<false, UserClass>::FMethodPtr InMethodPtr, FName InFunctionName)
+	template <typename UserClass, typename... VarTypes>
+	UE_REWRITE void __Internal_AddUniqueDynamic(TObjectPtr<UserClass> InUserObject, typename FDelegate::template TMethodPtrResolver<false, UserClass, VarTypes...>::FMethodPtr InMethodPtr, FName InFunctionName, VarTypes&&... Vars)
 	{
 		UE_STATIC_DEPRECATE(5.8, std::is_const_v<UserClass>, "Binding a delegate with a const object pointer and non-const function is deprecated.");
-		this->AddUniqueDynamicImpl(ToRawPtr(InUserObject), InFunctionName);
+		this->AddUniqueDynamicImpl(ToRawPtr(InUserObject), InFunctionName, Forward<VarTypes>(Vars)...);
 	}
-	template <typename UserClass>
-	UE_REWRITE void __Internal_AddUniqueDynamic(UserClass* InUserObject, typename FDelegate::template TMethodPtrResolver<true, UserClass>::FMethodPtr InMethodPtr, FName InFunctionName)
+	template <typename UserClass, typename... VarTypes>
+	UE_REWRITE void __Internal_AddUniqueDynamic(UserClass* InUserObject, typename FDelegate::template TMethodPtrResolver<true, UserClass, VarTypes...>::FMethodPtr InMethodPtr, FName InFunctionName, VarTypes&&... Vars)
 	{
-		this->AddUniqueDynamicImpl(InUserObject, InFunctionName);
+		this->AddUniqueDynamicImpl(InUserObject, InFunctionName, Forward<VarTypes>(Vars)...);
 	}
-	template <typename UserClass>
-	UE_REWRITE void __Internal_AddUniqueDynamic(TObjectPtr<UserClass> InUserObject, typename FDelegate::template TMethodPtrResolver<true, UserClass>::FMethodPtr InMethodPtr, FName InFunctionName)
+	template <typename UserClass, typename... VarTypes>
+	UE_REWRITE void __Internal_AddUniqueDynamic(TObjectPtr<UserClass> InUserObject, typename FDelegate::template TMethodPtrResolver<true, UserClass, VarTypes...>::FMethodPtr InMethodPtr, FName InFunctionName, VarTypes&&... Vars)
 	{
-		this->AddUniqueDynamicImpl(ToRawPtr(InUserObject), InFunctionName);
+		this->AddUniqueDynamicImpl(ToRawPtr(InUserObject), InFunctionName, Forward<VarTypes>(Vars)...);
 	}
 
 private:
-	template <typename UserClass>
-	void RemoveDynamicImpl(const UserClass* InUserObject, FName InFunctionName)
+	template <typename UserClass, typename... VarTypes>
+	void RemoveDynamicImpl(const UserClass* InUserObject, FName InFunctionName, VarTypes&&... Vars)
 	{
+		static_assert(UE_USE_DYNAMIC_DELEGATE_PAYLOADS || sizeof...(VarTypes) == 0, "Payload support for dynamic delegates is disabled");
+
 		check(InUserObject);
 
-		this->Remove(InUserObject, InFunctionName);
+		this->Remove(InUserObject, InFunctionName, Forward<VarTypes>(Vars)...);
 	}
 
 public:
@@ -1453,25 +1468,25 @@ public:
 	 *
 	 * NOTE:  We're not actually using the incoming method pointer.  We simply require it for type-safety reasons.
 	 */
-	template <typename UserClass>
-	UE_REWRITE void __Internal_RemoveDynamic(UserClass* InUserObject, typename FDelegate::template TMethodPtrResolver<false, UserClass>::FMethodPtr InMethodPtr, FName InFunctionName)
+	template <typename UserClass, typename... VarTypes>
+	UE_REWRITE void __Internal_RemoveDynamic(UserClass* InUserObject, typename FDelegate::template TMethodPtrResolver<false, UserClass, VarTypes...>::FMethodPtr InMethodPtr, FName InFunctionName, VarTypes&&... Vars)
 	{
-		this->RemoveDynamicImpl(InUserObject, InFunctionName);
+		this->RemoveDynamicImpl(InUserObject, InFunctionName, Forward<VarTypes>(Vars)...);
 	}
-	template <typename UserClass>
-	UE_REWRITE void __Internal_RemoveDynamic(TObjectPtr<UserClass> InUserObject, typename FDelegate::template TMethodPtrResolver<false, UserClass>::FMethodPtr InMethodPtr, FName InFunctionName)
+	template <typename UserClass, typename... VarTypes>
+	UE_REWRITE void __Internal_RemoveDynamic( TObjectPtr<UserClass> InUserObject, typename FDelegate::template TMethodPtrResolver<false, UserClass, VarTypes&&...>::FMethodPtr InMethodPtr, FName InFunctionName, VarTypes&&... Vars)
 	{
-		this->RemoveDynamicImpl(ToRawPtr(InUserObject), InFunctionName);
+		this->RemoveDynamicImpl(ToRawPtr(InUserObject), InFunctionName, Forward<VarTypes>(Vars)...);
 	}
-	template <typename UserClass>
-	UE_REWRITE void __Internal_RemoveDynamic(UserClass* InUserObject, typename FDelegate::template TMethodPtrResolver<true, UserClass>::FMethodPtr InMethodPtr, FName InFunctionName)
+	template <typename UserClass, typename... VarTypes>
+	UE_REWRITE void __Internal_RemoveDynamic(UserClass* InUserObject, typename FDelegate::template TMethodPtrResolver<true, UserClass, VarTypes...>::FMethodPtr InMethodPtr, FName InFunctionName, VarTypes&&... Vars)
 	{
-		this->RemoveDynamicImpl(InUserObject, InFunctionName);
+		this->RemoveDynamicImpl(InUserObject, InFunctionName, Forward<VarTypes>(Vars)...);
 	}
-	template <typename UserClass>
-	UE_REWRITE void __Internal_RemoveDynamic(TObjectPtr<UserClass> InUserObject, typename FDelegate::template TMethodPtrResolver<true, UserClass>::FMethodPtr InMethodPtr, FName InFunctionName)
+	template <typename UserClass, typename... VarTypes>
+	UE_REWRITE void __Internal_RemoveDynamic(TObjectPtr<UserClass> InUserObject, typename FDelegate::template TMethodPtrResolver<true, UserClass, VarTypes&&...>::FMethodPtr InMethodPtr, FName InFunctionName, VarTypes&&... Vars)
 	{
-		this->RemoveDynamicImpl(ToRawPtr(InUserObject), InFunctionName);
+		this->RemoveDynamicImpl(ToRawPtr(InUserObject), InFunctionName, Forward<VarTypes>(Vars)...);
 	}
 
 	/** Broadcasts this delegate to all bound objects, except to those that may have expired */
diff --git a/Engine/Source/Runtime/Core/Public/UObject/ScriptDelegates.h b/Engine/Source/Runtime/Core/Public/UObject/ScriptDelegates.h
--- a/Engine/Source/Runtime/Core/Public/UObject/ScriptDelegates.h
+++ b/Engine/Source/Runtime/Core/Public/UObject/ScriptDelegates.h
@@ -26,6 +26,12 @@
 	#endif
 #endif
 
+#if UE_USE_DYNAMIC_DELEGATE_PAYLOADS
+	#include "Concepts/Arithmetic.h"
+	#include "Concepts/UObject.h"
+	#include "Templates/Tuple.h"
+#endif
+
 namespace UE::Core::Private
 {
 	template <typename InThreadSafetyMode>
@@ -39,6 +45,125 @@ namespace UE::Core::Private
 		using UnicastThreadSafetyModeForMulticasts = FNotThreadSafeNotCheckedDelegateMode;
 	};
 
+#if UE_USE_DYNAMIC_DELEGATE_PAYLOADS
+	// Can't use TWeakObjectPtr directly here, so take a guess it being one by checking for one of its unique members
+	template <typename T>
+	concept CWeakPointerish =
+		requires(const T& Ptr)
+		{
+			Ptr.HasSameIndexAndSerialNumber(Ptr);
+		};
+
+	// This trait defines how different dynamic delegate payload arguments are captured.
+	// If the trait is not specialized for a particular type, that type is not currently supported as a payload argument.
+	template <typename T>
+	struct TDynamicDelegateParamPayloadCaptureType;
+
+	// Reference and cv-qualified arguments are captured by value and unqualified.
+	template <typename T>
+	struct TDynamicDelegateParamPayloadCaptureType<T&> : TDynamicDelegateParamPayloadCaptureType<T>
+	{
+	};
+	template <typename T>
+	struct TDynamicDelegateParamPayloadCaptureType<const T&> : TDynamicDelegateParamPayloadCaptureType<T>
+	{
+	};
+	template <typename T>
+	struct TDynamicDelegateParamPayloadCaptureType<const T> : TDynamicDelegateParamPayloadCaptureType<T>
+	{
+	};
+
+	// Allowed:
+	// - UObject pointers.
+	// - TWeakObjectPtr.
+	// - Numeric types.
+	// - void (only for return types, not parameters, but C++ will stop you having a void parameter before getting here)
+	// - Strings and FNames.
+	template <UE::CUObject T>    struct TDynamicDelegateParamPayloadCaptureType<T*>          { using Type = T*; };
+	template <UE::CArithmetic T> struct TDynamicDelegateParamPayloadCaptureType<T>           { using Type = T; };
+	template <CWeakPointerish T> struct TDynamicDelegateParamPayloadCaptureType<T>           { using Type = T; };
+	template <>                  struct TDynamicDelegateParamPayloadCaptureType<void>        { using Type = void; };
+	template <>                  struct TDynamicDelegateParamPayloadCaptureType<FAnsiString> { using Type = FAnsiString; };
+	template <>                  struct TDynamicDelegateParamPayloadCaptureType<FString>     { using Type = FString; };
+	template <>                  struct TDynamicDelegateParamPayloadCaptureType<FUtf8String> { using Type = FUtf8String; };
+	template <>                  struct TDynamicDelegateParamPayloadCaptureType<FName>       { using Type = FName; };
+
+	template <typename T>
+	using TDynamicDelegateParamPayloadCaptureType_T = typename TDynamicDelegateParamPayloadCaptureType<T>::Type;
+
+	// An interface to a dynamic delegate payload.  Type-erases payload operations.
+	template <typename WeakPtrType>
+	struct IDynamicDelegatePayload
+	{
+		// Calls ObjectPtr->ProcessEvent(Function, ParametersPlusPayload...) on ObjectPtr. 
+		virtual void ProcessEvent(WeakPtrType ObjectPtr, UFunction* Function, void* Parameters) const = 0;
+
+		virtual ~IDynamicDelegatePayload() = 0;
+	};
+
+	template <typename WeakPtrType>
+	UE_REWRITE IDynamicDelegatePayload<WeakPtrType>::~IDynamicDelegatePayload() = default;
+
+	template <typename WeakPtrType, typename MemFnType, typename... VarTypes>
+	struct TDynamicDelegatePayload;
+
+	template <typename WeakPtrType, typename... ParamTypes, typename... VarTypes>
+	struct TDynamicDelegatePayload<WeakPtrType, void(ParamTypes...), VarTypes...> : IDynamicDelegatePayload<WeakPtrType>
+	{
+		static_assert(sizeof...(VarTypes) != 0, "Allocating an object to hold an empty payload is wasteful");
+
+		template <typename... ArgTypes>
+		[[nodiscard]] UE_REWRITE explicit TDynamicDelegatePayload(ArgTypes&&... Args)
+			: Payload((ArgTypes&&)Args...)
+		{
+		}
+
+		void ProcessEvent(WeakPtrType ObjectPtr, UFunction* Function, void* Parameters) const override
+		{
+			// The incoming Parameters will point to a block of arguments which were passed to the dynamic delegate
+			// invocation.  These will be combined with the payload in order to produce a new parameters struct to
+			// be passed to the ProcessEvent call.
+
+			// If the payload replaces every argument, just pass the tuple itself as it should be the same layout.
+			if constexpr (sizeof...(ParamTypes) == sizeof...(VarTypes))
+			{
+				ObjectPtr.Get()->ProcessEvent(Function, (void*)&Payload);
+			}
+			// ... otherwise we need to combine Params and Payload into a single struct
+			else
+			{
+				using FullTupleType   = TTuple<std::decay_t<ParamTypes>...>;
+				using ParamsTupleType = TTuplePopBack_T<TTuple<std::decay_t<ParamTypes>...>, sizeof...(VarTypes)>;
+				using VarsTupleType   = TTuple<std::decay_t<VarTypes>...>;
+
+				((ParamsTupleType*)Parameters)->ApplyBefore(
+					[&Payload = Payload, ObjectPtr, Function](auto&&... Params)
+					{
+						Payload.ApplyAfter(
+							[&Payload = Payload, ObjectPtr, Function](auto&&... ParamsAndVars)
+							{
+								FullTupleType NewParameters(Forward<decltype(ParamsAndVars)>(ParamsAndVars)...);
+								ObjectPtr.Get()->ProcessEvent(Function, &NewParameters);
+							},
+							Forward<decltype(Params)>(Params)...
+						);
+					}
+				);
+			}
+		}
+
+	private:
+		// Mutable because ProcessEvent takes a non-const param pointer
+		mutable TTuple<VarTypes...> Payload;
+	};
+
+	template <typename WeakPtrType, typename BoundFuncType, typename... VarTypes>
+	UE_REWRITE TSharedPtr<const IDynamicDelegatePayload<WeakPtrType>> MakeDynamicDelegatePayload(VarTypes&&... Vars)
+	{
+		return MakeShared<TDynamicDelegatePayload<WeakPtrType, BoundFuncType, TDynamicDelegateParamPayloadCaptureType_T<VarTypes>...>>(Forward<VarTypes>(Vars)...);
+	}
+#endif
+
 	// Describes a type which is constructible with a ForceInit argument.
 	template <typename T>
 	concept CForceInitable = UE::CCompleteType<T> && std::is_constructible_v<T, EForceInit>;
@@ -205,7 +330,8 @@ public:
 	TScriptDelegate() 
 		: Object( nullptr )
 		, FunctionName( NAME_None )
-	{ }
+	{
+	}
 
 	TScriptDelegate(const TScriptDelegate& Other)
 	{
@@ -213,23 +339,40 @@ public:
 
 		Object = Other.Object;
 		FunctionName = Other.FunctionName;
+#if UE_USE_DYNAMIC_DELEGATE_PAYLOADS
+		Payload = Other.Payload;
+#endif
 	}
 
+	// TODO: We could add move operations to avoid the refcount bump of payloads, but it
+	// would require resetting the object and function name too.  This could break existing
+	// code which is relying on execute-after-move.  We don't want to support that, but we
+	// need a deprecation strategy.
+
 	TScriptDelegate& operator=(const TScriptDelegate& Other)
 	{
 		WeakPtrType OtherObject;
 		FName OtherFunctionName;
+#if UE_USE_DYNAMIC_DELEGATE_PAYLOADS
+		TSharedPtr<const UE::Core::Private::IDynamicDelegatePayload<WeakPtrType>> OtherPayload;
+#endif
 
 		{
 			FReadAccessScope OtherReadScope = Other.GetReadAccessScope();
 			OtherObject = Other.Object;
 			OtherFunctionName = Other.FunctionName;
+#if UE_USE_DYNAMIC_DELEGATE_PAYLOADS
+			OtherPayload = Other.Payload;
+#endif
 		}
 
 		{
 			FWriteAccessScope WriteScope = GetWriteAccessScope();
 			Object = OtherObject;
 			FunctionName = OtherFunctionName;
+#if UE_USE_DYNAMIC_DELEGATE_PAYLOADS
+			Payload = MoveTemp(OtherPayload);
+#endif
 		}
 
 		return *this;
@@ -265,6 +408,7 @@ public:
 
 		Object = InObject;
 		FunctionName = InFunctionName;
+		// TODO: Add payload support?  We don't have the dynamic delegate signature here to create a payload.
 	}
 
 	/** 
@@ -325,6 +469,9 @@ public:
 
 		Object = nullptr;
 		FunctionName = NAME_None;
+#if UE_USE_DYNAMIC_DELEGATE_PAYLOADS
+		Payload = nullptr;
+#endif
 	}
 
 	/**
@@ -346,6 +493,7 @@ public:
 		if( IsBound() )
 		{
 			FReadAccessScope ReadScope = GetReadAccessScope();
+			// TODO: Stringify payloads?
 			return ((UObjectTemplate*)GetUObject())->GetPathName() + TEXT(".") + GetFunctionName().ToString();
 		}
 		return TEXT( "<Unbound>" );
@@ -357,6 +505,7 @@ public:
 		FReadAccessScope ReadScope = D.GetReadAccessScope();
 
 		Ar << D.Object << D.FunctionName;
+		// TODO: Serialize payload
 		return Ar;
 	}
 
@@ -367,6 +516,7 @@ public:
 
 		FStructuredArchive::FRecord Record = Slot.EnterRecord();
 		Record << SA_VALUE(TEXT("Object"), D.Object) << SA_VALUE(TEXT("FunctionName"),D.FunctionName);
+		// TODO: Serialize payload
 	}
 
 	/** Comparison operators */
@@ -386,6 +536,7 @@ public:
 		{
 			FReadAccessScope ThisReadScope = GetReadAccessScope();
 			bResult = Object == OtherObject && FunctionName == OtherFunctionName;
+			// TODO: Compare payloads
 		}
 
 		return bResult;
@@ -514,13 +665,23 @@ public:
 		}
 
 		// Execute the delegate!
-		ObjectPtr->ProcessEvent(Function, Parameters);
+#if UE_USE_DYNAMIC_DELEGATE_PAYLOADS
+		if (this->Payload)
+		{
+			Payload->ProcessEvent(ObjectPtr, Function, Parameters);
+		}
+		else
+#endif
+		{
+			ObjectPtr->ProcessEvent(Function, Parameters);
+		}
 	}
 
 	[[nodiscard]] friend uint32 GetTypeHash(const TScriptDelegate& Delegate)
 	{
 		FReadAccessScope ReadScope = Delegate.GetReadAccessScope();
 
+		// TODO: Hash payloads
 		return HashCombine(GetTypeHash(Delegate.Object), GetTypeHash(Delegate.GetFunctionName()));
 	}
 
@@ -534,16 +695,32 @@ public:
 		TScriptDelegate Copy;
 		Copy.Object = Other.Object;
 		Copy.FunctionName = Other.FunctionName;
+#if UE_USE_DYNAMIC_DELEGATE_PAYLOADS
+		Copy.Payload = Other.Payload;
+#endif
 
 		return Copy;
 	}
 
 protected:
-	void BindProtected(WeakPtrType InObject, FName InFunctionName)
+#if UE_USE_DYNAMIC_DELEGATE_PAYLOADS
+	template <typename BoundFuncType, typename... VarTypes>
+	UE_REWRITE void BindProtected(WeakPtrType InObject, FName InFunctionName, VarTypes&&... Vars)
+	{
+		this->Object = InObject;
+		this->FunctionName = InFunctionName;
+		if constexpr (sizeof...(VarTypes) != 0)
+		{
+			this->Payload = UE::Core::Private::MakeDynamicDelegatePayload<WeakPtrType, BoundFuncType>(Forward<VarTypes>(Vars)...);
+		}
+	}
+#else
+	UE_REWRITE void BindProtected(WeakPtrType InObject, FName InFunctionName)
 	{
 		this->Object = InObject;
 		this->FunctionName = InFunctionName;
 	}
+#endif
 
 private:
 	/** The object bound to this delegate, or nullptr if no object is bound */
@@ -551,6 +728,9 @@ private:
 
 	/** Name of the function to call on the bound object */
 	FName FunctionName;
+#if UE_USE_DYNAMIC_DELEGATE_PAYLOADS
+	TSharedPtr<const UE::Core::Private::IDynamicDelegatePayload<WeakPtrType>> Payload;
+#endif
 
 	friend struct TIsZeroConstructType<TScriptDelegate>;
 };
@@ -670,6 +850,7 @@ public:
 			FReadAccessScope OtherReadScope = InDelegate.GetReadAccessScope();
 			Object = InDelegate.Object.Get();
 			FunctionName = InDelegate.FunctionName;
+			// TODO: Compare payloads
 		}
 
 		return Contains(Object, FunctionName);
@@ -684,6 +865,8 @@ public:
 	 */
 	bool Contains( const UObject* InObject, FName InFunctionName ) const
 	{
+		// TODO: Accept payloads
+
 		FReadAccessScope ReadScope = GetReadAccessScope();
 
 		return InvocationList.ContainsByPredicate( [=]( const UnicastDelegateType& Delegate ){
@@ -762,6 +945,8 @@ public:
 	 */
 	void Remove( const UObject* InObject, FName InFunctionName )
 	{
+		// TODO: Accept payloads
+
 		FWriteAccessScope WriteScope = GetWriteAccessScope();
 
 		// Remove the delegate
