Repair data where we might have introduce -0.0 in the data when we switched coordinate Left handed to right handed. This creates divergence in the replication layer as the server and client are not using the same cooking format.

#Synced-CL 49797767

#Synced-CL 49797767

diff --git a/Engine/Source/Runtime/Core/Public/UObject/FortniteMainBranchObjectVersions.inl b/Engine/Source/Runtime/Core/Public/UObject/FortniteMainBranchObjectVersions.inl
--- a/Engine/Source/Runtime/Core/Public/UObject/FortniteMainBranchObjectVersions.inl
+++ b/Engine/Source/Runtime/Core/Public/UObject/FortniteMainBranchObjectVersions.inl
@@ -742,5 +742,8 @@ FORTNITE_MAIN_VERSION(RootMotionModifierRotationOffset, 242)
 // Update PolyBezierSpline internal storage structure to use 3 control points per segment instead of 4
 FORTNITE_MAIN_VERSION(PolyBezierSplineAnchorTripleLayout, 243)
 
+// Repairing verse data where we introduced -0.0 value when we flipped the coordinate system
+FORTNITE_MAIN_VERSION(VerseRepairDataWithMinusZeroValue, 244)
+
 // -----<new versions can be added above this line>-------------------------------------------------
 #undef FORTNITE_MAIN_VERSION
