SharedMutex: Added comments explaining the use of release memory order in UnlockShared()

#rnx

#Synced-CL 49858055

diff --git a/Engine/Source/Runtime/Core/Private/Async/SharedRecursiveMutex.cpp b/Engine/Source/Runtime/Core/Private/Async/SharedRecursiveMutex.cpp
--- a/Engine/Source/Runtime/Core/Private/Async/SharedRecursiveMutex.cpp
+++ b/Engine/Source/Runtime/Core/Private/Async/SharedRecursiveMutex.cpp
@@ -315,6 +315,8 @@ FORCENOINLINE void FSharedRecursiveMutex::LockSharedSlow(Core::Private::FSharedR
 void FSharedRecursiveMutex::UnlockShared(Core::Private::FSharedRecursiveMutexLink& Link)
 {
 	Link.Pop();
+	// C++ SharedMutex requires that UnlockShared() synchronizes-with subsequent Lock() operations.
+	// This requirement necessitates the use of release memory order or stronger.
 	const uint32 LastState = State.fetch_sub(1 << SharedLockCountShift, std::memory_order_release);
 	checkSlow(LastState & SharedLockCountMask);
 	constexpr uint32 WakeState = MayHaveWaitingLockFlag | (1 << SharedLockCountShift);
diff --git a/Engine/Source/Runtime/Core/Public/Async/SharedMutex.h b/Engine/Source/Runtime/Core/Public/Async/SharedMutex.h
--- a/Engine/Source/Runtime/Core/Public/Async/SharedMutex.h
+++ b/Engine/Source/Runtime/Core/Public/Async/SharedMutex.h
@@ -55,8 +55,8 @@ public:
 
 	inline void Unlock()
 	{
-		// Unlock immediately to allow other threads to acquire the lock while this thread looks for a thread to wake.
 		UE_AUTORTFM_REPORT_HAZARD_IF_CLOSED(SharedMutex);
+		// Unlock immediately to allow other threads to acquire the lock while this thread looks for a thread to wake.
 		uint32 LastState = State.fetch_sub(IsLockedFlag, std::memory_order_release);
 		checkSlow(LastState & IsLockedFlag);
 		if (LIKELY(!(LastState & (MayHaveWaitingLockFlag | MayHaveWaitingSharedLockFlag))))
@@ -102,8 +102,10 @@ public:
 
 	inline void UnlockShared()
 	{
-		// Unlock immediately to allow other threads to acquire the lock while this thread looks for a thread to wake.
 		UE_AUTORTFM_REPORT_HAZARD_IF_CLOSED(SharedMutex);
+		// Unlock immediately to allow other threads to acquire the lock while this thread looks for a thread to wake.
+		// C++ SharedMutex requires that UnlockShared() synchronizes-with subsequent Lock() operations.
+		// This requirement necessitates the use of release memory order or stronger.
 		const uint32 LastState = State.fetch_sub(1 << SharedLockCountShift, std::memory_order_release);
 		checkSlow(LastState & SharedLockCountMask);
 		constexpr uint32 WakeState = MayHaveWaitingLockFlag | (1 << SharedLockCountShift);
