[AutoRTFM] Provide a no-op implementation of PrintCallstackOnCrash() for shipping builds

Apparently we build some shipping targets of our test executables.

#Synced-CL 49880520

diff --git a/Engine/Source/Runtime/Core/Public/AutoRTFM/AutoRTFMTestingUE.h b/Engine/Source/Runtime/Core/Public/AutoRTFM/AutoRTFMTestingUE.h
--- a/Engine/Source/Runtime/Core/Public/AutoRTFM/AutoRTFMTestingUE.h
+++ b/Engine/Source/Runtime/Core/Public/AutoRTFM/AutoRTFMTestingUE.h
@@ -4,15 +4,17 @@
 
 #include "AutoRTFM/Testing.h"
 
-#if AUTORTFM_ENABLE_TEST_UTILS
 
 namespace AutoRTFM::Testing
 {
 
+#if AUTORTFM_ENABLE_TEST_UTILS
 // Installs a crash handler to print the callstack if the process crashes.
 // Can only be called once for the process.
 AUTORTFM_DISABLE CORE_API void PrintCallstackOnCrash();
+#else
+inline void PrintCallstackOnCrash() {}
+#endif
 
 }
 
-#endif // AUTORTFM_ENABLE_TEST_UTILS
