Stored UObject*s as TWeakObjectPtrs in dynamic delegate payloads.

diff --git a/Engine/Source/Runtime/Core/Public/UObject/ScriptDelegates.h b/Engine/Source/Runtime/Core/Public/UObject/ScriptDelegates.h
--- a/Engine/Source/Runtime/Core/Public/UObject/ScriptDelegates.h
+++ b/Engine/Source/Runtime/Core/Public/UObject/ScriptDelegates.h
@@ -30,6 +30,7 @@
 	#include "Concepts/Arithmetic.h"
 	#include "Concepts/UObject.h"
 	#include "Templates/Tuple.h"
+	#include "UObject/WeakObjectPtrTemplatesFwd.h"
 #endif
 
 namespace UE::Core::Private
@@ -54,22 +55,39 @@ namespace UE::Core::Private
 			Ptr.HasSameIndexAndSerialNumber(Ptr);
 		};
 
+	// Wrapper for TWeakObjectPtr that implicitly converts to and from its raw pointer, for use in payload captures.
+	template <typename T, typename WeakPtrType>
+	struct TImplicitWeakObjectPtr
+	{
+		TImplicitWeakObjectPtr(T* Ptr)
+			: WeakObjPtr(Ptr)
+		{
+		}
+
+		UE_REWRITE operator T*() const
+		{
+			return this->WeakObjPtr.Get();
+		}
+
+		TWeakObjectPtr<T, WeakPtrType> WeakObjPtr;
+	};
+
 	// This trait defines how different dynamic delegate payload arguments are captured.
 	// If the trait is not specialized for a particular type, that type is not currently supported as a payload argument.
-	template <typename T>
+	template <typename WeakPtrType, typename T>
 	struct TDynamicDelegateParamPayloadCaptureType;
 
 	// Reference and cv-qualified arguments are captured by value and unqualified.
-	template <typename T>
-	struct TDynamicDelegateParamPayloadCaptureType<T&> : TDynamicDelegateParamPayloadCaptureType<T>
+	template <typename WeakPtrType, typename T>
+	struct TDynamicDelegateParamPayloadCaptureType<WeakPtrType, T&> : TDynamicDelegateParamPayloadCaptureType<WeakPtrType, T>
 	{
 	};
-	template <typename T>
-	struct TDynamicDelegateParamPayloadCaptureType<const T&> : TDynamicDelegateParamPayloadCaptureType<T>
+	template <typename WeakPtrType, typename T>
+	struct TDynamicDelegateParamPayloadCaptureType<WeakPtrType, const T&> : TDynamicDelegateParamPayloadCaptureType<WeakPtrType, T>
 	{
 	};
-	template <typename T>
-	struct TDynamicDelegateParamPayloadCaptureType<const T> : TDynamicDelegateParamPayloadCaptureType<T>
+	template <typename WeakPtrType, typename T>
+	struct TDynamicDelegateParamPayloadCaptureType<WeakPtrType, const T> : TDynamicDelegateParamPayloadCaptureType<WeakPtrType, T>
 	{
 	};
 
@@ -79,17 +97,17 @@ namespace UE::Core::Private
 	// - Numeric types.
 	// - void (only for return types, not parameters, but C++ will stop you having a void parameter before getting here)
 	// - Strings and FNames.
-	template <UE::CUObject T>    struct TDynamicDelegateParamPayloadCaptureType<T*>          { using Type = T*; };
-	template <UE::CArithmetic T> struct TDynamicDelegateParamPayloadCaptureType<T>           { using Type = T; };
-	template <CWeakPointerish T> struct TDynamicDelegateParamPayloadCaptureType<T>           { using Type = T; };
-	template <>                  struct TDynamicDelegateParamPayloadCaptureType<void>        { using Type = void; };
-	template <>                  struct TDynamicDelegateParamPayloadCaptureType<FAnsiString> { using Type = FAnsiString; };
-	template <>                  struct TDynamicDelegateParamPayloadCaptureType<FString>     { using Type = FString; };
-	template <>                  struct TDynamicDelegateParamPayloadCaptureType<FUtf8String> { using Type = FUtf8String; };
-	template <>                  struct TDynamicDelegateParamPayloadCaptureType<FName>       { using Type = FName; };
-
-	template <typename T>
-	using TDynamicDelegateParamPayloadCaptureType_T = typename TDynamicDelegateParamPayloadCaptureType<T>::Type;
+	template <typename WeakPtrType, UE::CUObject T>    struct TDynamicDelegateParamPayloadCaptureType<WeakPtrType, T*>          { using Type = TImplicitWeakObjectPtr<T, WeakPtrType>; };
+	template <typename WeakPtrType, UE::CArithmetic T> struct TDynamicDelegateParamPayloadCaptureType<WeakPtrType, T>           { using Type = T; };
+	template <typename WeakPtrType, CWeakPointerish T> struct TDynamicDelegateParamPayloadCaptureType<WeakPtrType, T>           { using Type = T; };
+	template <typename WeakPtrType>                    struct TDynamicDelegateParamPayloadCaptureType<WeakPtrType, void>        { using Type = void; };
+	template <typename WeakPtrType>                    struct TDynamicDelegateParamPayloadCaptureType<WeakPtrType, FAnsiString> { using Type = FAnsiString; };
+	template <typename WeakPtrType>                    struct TDynamicDelegateParamPayloadCaptureType<WeakPtrType, FString>     { using Type = FString; };
+	template <typename WeakPtrType>                    struct TDynamicDelegateParamPayloadCaptureType<WeakPtrType, FUtf8String> { using Type = FUtf8String; };
+	template <typename WeakPtrType>                    struct TDynamicDelegateParamPayloadCaptureType<WeakPtrType, FName>       { using Type = FName; };
+
+	template <typename WeakPtrType, typename T>
+	using TDynamicDelegateParamPayloadCaptureType_T = typename TDynamicDelegateParamPayloadCaptureType<WeakPtrType, T>::Type;
 
 	// An interface to a dynamic delegate payload.  Type-erases payload operations.
 	template <typename WeakPtrType>
@@ -160,7 +178,7 @@ namespace UE::Core::Private
 	template <typename WeakPtrType, typename BoundFuncType, typename... VarTypes>
 	UE_REWRITE TSharedPtr<const IDynamicDelegatePayload<WeakPtrType>> MakeDynamicDelegatePayload(VarTypes&&... Vars)
 	{
-		return MakeShared<TDynamicDelegatePayload<WeakPtrType, BoundFuncType, TDynamicDelegateParamPayloadCaptureType_T<VarTypes>...>>(Forward<VarTypes>(Vars)...);
+		return MakeShared<TDynamicDelegatePayload<WeakPtrType, BoundFuncType, TDynamicDelegateParamPayloadCaptureType_T<WeakPtrType, VarTypes>...>>(Forward<VarTypes>(Vars)...);
 	}
 #endif
 
