LLM: Implement ELLMTagSets::CodeOrContent, to track whether allocations are made inside any ELLMTagSet::Asset scope without recording which asset the scope was created for. This supports making a memory budget for how much content is loaded separately from a memory budget for the global data used by the code system that supports the content.
#rnx

#Synced-CL 49850719

diff --git a/Engine/Source/Runtime/Core/Private/HAL/LLM/LLM.cpp b/Engine/Source/Runtime/Core/Private/HAL/LLM/LLM.cpp
--- a/Engine/Source/Runtime/Core/Private/HAL/LLM/LLM.cpp
+++ b/Engine/Source/Runtime/Core/Private/HAL/LLM/LLM.cpp
@@ -1012,6 +1012,10 @@ FAutoConsoleCommand DumpLLM(
 		{
 			TagSet = ELLMTagSet::UObjectClasses;
 		}
+		else if (FParse::Param(*Command, TEXT("CodeOrContent")))
+		{
+			TagSet = ELLMTagSet::CodeOrContent;
+		}
 
 		FLowLevelMemTracker::Get().DumpToLog(DumpFormat, &Ar, SizeParams, TagSet);
 	}));
@@ -1343,6 +1347,8 @@ static FName TagName_Untagged("Untagged");
 static FName TagName_UntaggedAsset("UntaggedAsset");
 static FName TagName_UntaggedAssetClass("UntaggedAssetClass");
 static FName TagName_UntaggedUObjectClasses("UntaggedUObjectClasses");
+LLM_DEFINE_TAG(CodeOrContent_Code, FName(TEXT("Code")), NAME_None, NAME_None, NAME_None, ELLMTagSet::CodeOrContent);
+LLM_DEFINE_TAG(CodeOrContent_Content, FName(TEXT("Content")), NAME_None, NAME_None, NAME_None, ELLMTagSet::CodeOrContent);
 
 FName LLMGetTagUniqueName(ELLMTag Tag)
 {
@@ -1424,6 +1430,10 @@ extern const FName LLMGetUntaggedTagName(ELLMTagSet TagSet)
 		{
 			return TagName_UntaggedUObjectClasses;
 		}
+		case ELLMTagSet::CodeOrContent:
+		{
+			return LLMTagDeclaration_CodeOrContent_Code.GetUniqueName();
+		}
 		default:
 		{
 			return NAME_None;
@@ -1614,7 +1624,9 @@ bool FLowLevelMemTracker::IsTagSetActive(ELLMTagSet Set)
 	{
 		return false;
 	}
-	return GetGlobals()->IsTagSetActiveInner(Set);
+	UE::LLMPrivate::FLLMGlobals* Globals = GetGlobals();
+	Globals->BootstrapInitialise();
+	return Globals->IsTagSetActiveInner(Set);
 }
 
 bool FLowLevelMemTracker::ShouldReduceThreads()
@@ -2411,7 +2423,7 @@ void FLLMGlobals::ProcessCommandLineInner(const TCHAR* CmdLine)
 	bool LocalActiveSets[static_cast<uint8>(ELLMTagSet::Max)] = { false };
 	LocalActiveSets[static_cast<uint8>(ELLMTagSet::None)] = true;
 
-	static_assert(static_cast<uint8>(ELLMTagSet::Max) == 4, "You added a tagset, without updating FLowLevelMemTracker::ProcessCommandLine");
+	static_assert(static_cast<uint8>(ELLMTagSet::Max) == 5, "You added a tagset without updating FLowLevelMemTracker::ProcessCommandLine");
 	if (FParse::Value(CmdLine, TEXT("LLMTAGSETS="), SetList, false /* bShouldStopOnSeparator */))
 	{
 #if LLM_ALLOW_NAMES_TAGS
@@ -2431,6 +2443,10 @@ void FLLMGlobals::ProcessCommandLineInner(const TCHAR* CmdLine)
 			{
 				LocalActiveSets[static_cast<int32>(ELLMTagSet::UObjectClasses)] = true;
 			}
+			else if (Set == TEXT("CodeOrContent"))
+			{
+				LocalActiveSets[static_cast<int32>(ELLMTagSet::CodeOrContent)] = true;
+			}
 		}
 #else
 		// We do not support activating ELLMTagSets other than ELLMTagSet::None when not using LLM_ALLOW_NAMES_TAGS,
@@ -2754,8 +2770,25 @@ bool FLLMGlobals::ExecInner(const TCHAR* Cmd, FOutputDevice& Ar)
 
 bool FLLMGlobals::IsTagSetActiveInner(ELLMTagSet Set)
 {
-	BootstrapInitialise();
+	return IsTagSetScopeActiveInner(Set);
+}
+
+bool FLLMGlobals::IsTagSetScopeActiveInner(ELLMTagSet Set)
+{
+	// ELLMTagSets::AssetClasses and ELLMTagSets::CodeOrContent share scopes
+	if (Set == ELLMTagSet::AssetClasses || Set == ELLMTagSet::CodeOrContent)
+	{
+		return ActiveSets[static_cast<int32>(ELLMTagSet::AssetClasses)]
+			|| ActiveSets[static_cast<int32>(ELLMTagSet::CodeOrContent)];
+	}
+	else
+	{
+		return ActiveSets[static_cast<int32>(Set)];
+	}
+}
 
+bool FLLMGlobals::IsTagSetRecordingActiveInner(ELLMTagSet Set)
+{
 	return ActiveSets[static_cast<int32>(Set)];
 }
 
@@ -2766,7 +2799,7 @@ bool FLLMGlobals::ShouldReduceThreadsInner()
 	LLMCheckf(bConfigurationComplete,
 		TEXT("ShouldReduceThreads has been called too early, before we processed the configuration settings required for it."));
 
-	return IsTagSetActive(ELLMTagSet::Assets) || IsTagSetActive(ELLMTagSet::AssetClasses);
+	return IsTagRecordingActiveInner(ELLMTagSet::Assets) || IsTagRecordingActiveInner(ELLMTagSet::AssetClasses);
 #else
 	return false;
 #endif
@@ -3034,7 +3067,10 @@ FTagData& FLLMGlobals::RegisterTagData(FName Name, FName DisplayName, FName Pare
 	else if (TagSet != ELLMTagSet::None)
 	{
 		// Tags not in the None TagSet do not have parents and do not require name validation
-		DisplayName = Name;
+		if (DisplayName.IsNone())
+		{
+			DisplayName = Name;
+		}
 	}
 	else if (bIsStatTag)
 	{
@@ -3042,7 +3078,10 @@ FTagData& FLLMGlobals::RegisterTagData(FName Name, FName DisplayName, FName Pare
 		// Stat tags do not specify their parent, and their parent is set to the CustomName aggregator.
 		LLMCheck(DisplayName.IsNone());
 		LLMCheck(ParentName.IsNone());
-		DisplayName = Name;
+		if (DisplayName.IsNone())
+		{
+			DisplayName = Name;
+		}
 		ParentName = GetTagName_CustomName();
 	}
 	else
@@ -3627,11 +3666,25 @@ const FTagData* FLLMGlobals::FindOrAddDefaultTagData(ELLMTagSet TagSet)
 	case ELLMTagSet::UObjectClasses:
 		return FindOrAddTagData(TagName_UntaggedUObjectClasses, TagSet, false /* bIsStatTag */,
 			ETagReferenceSource::FunctionAPI);
+	case ELLMTagSet::CodeOrContent:
+		return FindOrAddCodeOrContentCodeTagData();
 	default:
 		return nullptr;
 	}
 }
 
+const FTagData* FLLMGlobals::FindOrAddCodeOrContentCodeTagData()
+{
+	return FindOrAddTagData(LLMTagDeclaration_CodeOrContent_Code.GetUniqueName(), ELLMTagSet::CodeOrContent,
+		false /* bIsStatTag */, ETagReferenceSource::FunctionAPI);
+}
+
+const FTagData* FLLMGlobals::FindOrAddCodeOrContentContentTagData()
+{
+	return FindOrAddTagData(LLMTagDeclaration_CodeOrContent_Content.GetUniqueName(), ELLMTagSet::CodeOrContent,
+		false /* bIsStatTag */, ETagReferenceSource::FunctionAPI);
+}
+
 FAllocationGroup* FLLMGlobals::FindOrAddAllocationGroup(const FActiveTags& ActiveTags)
 {
 	// Caller must hold write lock on AllocationGroupsMutex.
@@ -3742,7 +3795,7 @@ FLLMScope::FLLMScope(FName TagName, bool bInIsStatTag, ELLMTagSet InTagSet, ELLM
 		if (FLowLevelMemTracker::IsEnabled())
 		{
 			LLMRef.BootstrapInitialise();
-			if (LLMRef.IsTagSetActive(InTagSet))
+			if (LLMRef.IsTagSetScopeActiveInner(InTagSet))
 			{
 				if (bOverride || LLMRef.GetTracker(InTracker)->GetActiveTagData(InTagSet) == nullptr)
 				{
@@ -3763,6 +3816,8 @@ FLLMScope::FLLMScope(FName TagName, bool bInIsStatTag, ELLMTagSet InTagSet, ELLM
 FLLMScope::FLLMScope(const UE::LLMPrivate::FTagData* TagData, bool bInIsStatTag, ELLMTagSet InTagSet,
 	ELLMTracker InTracker, bool bOverride)
 {
+	using namespace UE::LLMPrivate;
+
 	if (!FLowLevelMemTracker::IsEnabled())
 	{
 		return;
@@ -3784,6 +3839,21 @@ FLLMScope::FLLMScope(const UE::LLMPrivate::FTagData* TagData, bool bInIsStatTag,
 				Tracker = InTracker;
 				TagSet = InTagSet;
 				LLMRef.GetTracker(Tracker)->PushTag(TagData, InTagSet);
+				if (TagSet == ELLMTagSet::AssetClasses && LLMRef.IsTagSetRecordingActiveInner(ELLMTagSet::CodeOrContent))
+				{
+					// AssetClasses scopes also have the responsibility to trigger CodeOrContent scopes, so when
+					// pushing onto AssetClasses,onto CodeOrContent.
+					bPushedCodeOrContent = true;
+					if (!TagData)
+					{
+						LLMRef.GetTracker(Tracker)->PushTag(nullptr, ELLMTagSet::CodeOrContent);
+					}
+					else
+					{
+						const FTagData* CodeOrContentData = LLMRef.FindOrAddCodeOrContentContentTagData();
+						LLMRef.GetTracker(Tracker)->PushTag(CodeOrContentData, ELLMTagSet::CodeOrContent);
+					}
+				}
 			}
 		}
 	};
@@ -3794,8 +3864,15 @@ FLLMScope::FLLMScope(const UE::LLMPrivate::FTagData* TagData, bool bInIsStatTag,
 
 void FLLMScope::Destruct()
 {
-	UE::LLMPrivate::FLLMGlobals& LLMRef = UE::LLMPrivate::FLLMGlobals::GetInner();
-	LLMRef.GetTracker(Tracker)->PopTag(TagSet);
+	using namespace UE::LLMPrivate;
+
+	FLLMGlobals& LLMRef = UE::LLMPrivate::FLLMGlobals::GetInner();
+	FLLMTracker* TrackerPtr = LLMRef.GetTracker(Tracker);
+	TrackerPtr->PopTag(TagSet);
+	if (bPushedCodeOrContent)
+	{
+		TrackerPtr->PopTag(ELLMTagSet::CodeOrContent);
+	}
 }
 
 FLLMScope::~FLLMScope()
@@ -3820,22 +3897,39 @@ void FLLMScopeDynamic::Init(ELLMTracker InTracker, ELLMTagSet InTagSet)
 		return;
 	}
 	LLMRef.BootstrapInitialise();
-	if (!LLMRef.IsTagSetActive(InTagSet))
+	if (!LLMRef.IsTagSetScopeActiveInner(InTagSet))
 	{
 		bEnabled = false;
 		return;
 	}
 	bEnabled = true;
 	TagData = nullptr;
+	CodeOrContentTagData = nullptr;
 	Tracker = InTracker;
 	TagSet = InTagSet;
 }
 
 bool FLLMScopeDynamic::TryFindTag(FName TagName)
 {
+	// AssetClasses scopes also have the responsibility to trigger CodeOrContent scopes, so we might
+	// reach this point even if TagSet is not RecordingActive. In the case where the TagSet is not
+	// recording, we should not create its tag and should return true from TryFindTag.
+	bool bNeedAddTag = false;
 	UE::LLMPrivate::FLLMGlobals& LLMRef = UE::LLMPrivate::FLLMGlobals::GetInner();
-	TagData = LLMRef.FindTagData(TagName, TagSet, UE::LLMPrivate::ETagReferenceSource::Scope);
-	return TagData != nullptr;
+	if (LLMRef.IsTagSetRecordingActiveInner(TagSet))
+	{
+		TagData = LLMRef.FindTagData(TagName, TagSet, UE::LLMPrivate::ETagReferenceSource::Scope);
+		bNeedAddTag = TagData == nullptr;
+	}
+	if (TagSet == ELLMTagSet::AssetClasses && LLMRef.IsTagSetRecordingActiveInner(ELLMTagSet::CodeOrContent))
+	{
+		CodeOrContentTagData = LLMRef.FindOrAddCodeOrContentContentTagData();
+	}
+	else
+	{
+		CodeOrContentTagData = nullptr;
+	}
+	return !bNeedAddTag;
 }
 
 bool FLLMScopeDynamic::TryAddTagAndActivate(FName TagName, const ILLMDynamicTagConstructor& Constructor)
@@ -3865,6 +3959,11 @@ bool FLLMScopeDynamic::TryAddTagAndActivate(FName TagName, const ILLMDynamicTagC
 			case ELLMTagSet::UObjectClasses:
 				StatFullName = FDynamicStats::CreateMemoryStatId<FStatGroup_STATGROUP_LLMUObjects>(
 					StatConstructorName).GetName();
+				break;
+			case ELLMTagSet::CodeOrContent:
+				StatFullName = FDynamicStats::CreateMemoryStatId<FStatGroup_STATGROUP_LLMCodeOrContent>(
+					StatConstructorName).GetName();
+				break;
 			default:
 				checkNoEntry();
 				break;
@@ -3878,22 +3977,37 @@ bool FLLMScopeDynamic::TryAddTagAndActivate(FName TagName, const ILLMDynamicTagC
 #endif
 	TagData = LLMRef.FindOrAddTagData(TagName, TagSet, StatFullName, UE::LLMPrivate::ETagReferenceSource::Scope);
 	LLMRef.GetTracker(Tracker)->PushTag(TagData, TagSet);
+	if (CodeOrContentTagData)
+	{
+		LLMRef.GetTracker(Tracker)->PushTag(CodeOrContentTagData, ELLMTagSet::CodeOrContent);
+	}
 	return true;
 }
 
 void FLLMScopeDynamic::Activate()
 {
 	UE::LLMPrivate::FLLMGlobals& LLMRef = UE::LLMPrivate::FLLMGlobals::GetInner();
-	LLMRef.GetTracker(Tracker)->PushTag(TagData, TagSet);
+	if (TagData)
+	{
+		LLMRef.GetTracker(Tracker)->PushTag(TagData, TagSet);
+	}
+	if (CodeOrContentTagData)
+	{
+		LLMRef.GetTracker(Tracker)->PushTag(CodeOrContentTagData, ELLMTagSet::CodeOrContent);
+	}
 }
 
 void FLLMScopeDynamic::Destruct()
 {
+	UE::LLMPrivate::FLLMGlobals& LLMRef = UE::LLMPrivate::FLLMGlobals::GetInner();
 	if (TagData)
 	{
-		UE::LLMPrivate::FLLMGlobals& LLMRef = UE::LLMPrivate::FLLMGlobals::GetInner();
 		LLMRef.GetTracker(Tracker)->PopTag(TagSet);
 	}
+	if (CodeOrContentTagData)
+	{
+		LLMRef.GetTracker(Tracker)->PopTag(ELLMTagSet::CodeOrContent);
+	}
 }
 
 FLLMPauseScope::FLLMPauseScope(FName TagName, bool bIsStatTag, uint64 Amount, ELLMTracker TrackerToPause,
@@ -3931,7 +4045,7 @@ void FLLMPauseScope::Init(FName TagName, ELLMTag EnumTag, bool bIsEnumTag, bool
 		return;
 	}
 	LLMRef.BootstrapInitialise();
-	if (!LLMRef.IsTagSetActive(ELLMTagSet::None))
+	if (!LLMRef.IsTagSetScopeActiveInner(ELLMTagSet::None))
 	{
 		bEnabled = false;
 		return;
@@ -4028,7 +4142,7 @@ FLLMScopeFromPtr::FLLMScopeFromPtr(void* Ptr, ELLMTracker InTracker)
 
 	for (int32 TagSetAsInteger = 0; TagSetAsInteger < static_cast<int32>(ELLMTagSet::Max); TagSetAsInteger++)
 	{
-		if (!LLMRef.IsTagSetActive(static_cast<ELLMTagSet>(TagSetAsInteger)))
+		if (!LLMRef.IsTagSetScopeActiveInner(static_cast<ELLMTagSet>(TagSetAsInteger)))
 		{
 			bEnabled[TagSetAsInteger] = false;
 
@@ -6744,12 +6858,23 @@ void FLLMThreadState::ClearAllocTypeAmounts()
 }
 
 // FLLMCsvWriter implementation.
+bool FLLMCsvWriter::FColumnKey::operator==(const FLLMCsvWriter::FColumnKey& Other) const
+{
+	return SystemsTag == Other.SystemsTag && CodeOrContentTag == Other.CodeOrContentTag;
+}
+
+uint32 GetTypeHash(const FLLMCsvWriter::FColumnKey& Key)
+{
+	return ::GetTypeHash(Key.SystemsTag) * 101 + ::GetTypeHash(Key.CodeOrContentTag);
+}
+
 FLLMCsvWriter::FLLMCsvWriter()
 	: Archive(nullptr)
 	, LastWriteTime(FPlatformTime::Seconds())
 	, WriteCount(0)
 	, HeaderMaxSize(-1)
 	, bRegisteredFlushDelegate(false)
+	, bRecordingCodeOrContent(false)
 {
 }
 
@@ -6761,7 +6886,6 @@ FLLMCsvWriter::~FLLMCsvWriter()
 void FLLMCsvWriter::Clear()
 {
 	Columns.Empty();
-	ExistingColumns.Empty();
 }
 
 void FLLMCsvWriter::OnPreFork()
@@ -6806,14 +6930,14 @@ void FLLMCsvWriter::Publish(FLLMGlobals& LLMRef, const FTrackerTagSizeMap& TagSi
 
 	LastWriteTime = Now;
 
-	const bool bCreatedArchive = CreateArchive();
-	const bool bColumnsUpdated = UpdateColumns(TagSizes);
+	const bool bCreatedArchive = CreateArchive(LLMRef);
+	const bool bColumnsUpdated = UpdateColumns(LLMRef, TagSizes);
 
 	if (bCreatedArchive || bColumnsUpdated)
 	{
 		// The column names are written at the start of the archive; when they change we seek back to the start of
 		// the file and rewrite the column names.
-		WriteHeader(OverrideTrackedTotalTagData, OverrideUntaggedTagData);
+		WriteHeader(LLMRef, OverrideTrackedTotalTagData, OverrideUntaggedTagData);
 	}
 
 	AddRow(LLMRef, TagSizes, OverrideTrackedTotalTagData, OverrideUntaggedTagData, TrackedTotal, SizeParams);
@@ -6836,13 +6960,17 @@ void FLLMCsvWriter::Write(FStringView Text)
 		StringCast<ANSICHAR>(Text.GetData(), Text.Len()).Get()), Text.Len() * sizeof(ANSICHAR));
 }
 
-bool FLLMCsvWriter::CreateArchive()
+bool FLLMCsvWriter::CreateArchive(FLLMGlobals& LLMRef)
 {
 	if (Archive)
 	{
 		return false;
 	}
 
+	// Initialize some variables that require ActiveSets to have been parsed.
+	bRecordingCodeOrContent = Tracker == ELLMTracker::Default
+		&& LLMRef.IsTagSetRecordingActiveInner(ELLMTagSet::CodeOrContent);
+
 	// Create the csv file.
 	FString Directory = UE::LLM::GetLLMProfilingDir();
 	IFileManager::Get().MakeDirectory(*Directory, true);
@@ -6874,10 +7002,19 @@ bool FLLMCsvWriter::CreateArchive()
 	}
 	else
 	{
+		int32 CodeOrContentCodeLength = 0;
+		int32 CodeOrContentContentLength = 0;
+		if (bRecordingCodeOrContent)
+		{
+			const FTagData* CodeData = LLMRef.FindOrAddCodeOrContentCodeTagData();
+			const FTagData* ContentData = LLMRef.FindOrAddCodeOrContentContentTagData();
+			CodeOrContentCodeLength = WriteToString<256>(CodeData->GetDisplayName()).Len();
+			CodeOrContentContentLength = WriteToString<256>(ContentData->GetDisplayName()).Len();
+		}
+
 		HeaderMaxSize = 0;
 		constexpr int32 MaxNameStrLen = 256;
 		TStringBuilder<MaxNameStrLen> TagNameStr;
-		FLLMGlobals& LLMRef = FLLMGlobals::GetInner();
 		// Find the name size of all reportable TagDatas. This will not include TagDatas that have not yet been created,
 		// but it will gather all of the compiled-in ELLMTag and FName tags. We only allow reporting
 		// ELLMTagSet::None, and not many dynamic tags exist in ELLMTagSet::None, so this will reserve space for most of
@@ -6890,7 +7027,13 @@ bool FLLMCsvWriter::CreateArchive()
 				TagNameStr.Reset();
 				// We have to avoid allocating memory in this loop because we hold an LLM lock.
 				TagData->GetDisplayPath(TagNameStr, MaxNameStrLen - 1);
-				HeaderMaxSize += TagNameStr.Len() + 1;
+				int32 SizeForTag = TagNameStr.Len() + 1;
+				if (bRecordingCodeOrContent)
+				{
+					SizeForTag = SizeForTag + 1 + CodeOrContentCodeLength +
+						SizeForTag + 1 + CodeOrContentContentLength;
+				}
+				HeaderMaxSize += SizeForTag;
 			}
 		}
 	}
@@ -6901,57 +7044,130 @@ bool FLLMCsvWriter::CreateArchive()
 	return true;
 }
 
-bool FLLMCsvWriter::UpdateColumns(const FTrackerTagSizeMap& TagSizes)
+bool FLLMCsvWriter::UpdateColumns(FLLMGlobals& LLMRef, const FTrackerTagSizeMap& TagSizes)
 {
 	bool bUpdated = false;
 
-	for (const TPair<const FTagData*, FTrackerTagSizeData>& It : TagSizes)
+	if (bRecordingCodeOrContent)
 	{
-		const FTagData* TagData = It.Key;
-		if (!TagData->IsReportable())
+		// When recording ELLMTagSets::CodeOrContent, we write a column for every combination of ELLMTagSet::None
+		// and ELLMTagSet::CodeOrContent, rather than a column for every TagData. To find the sizes of the 
+		// combinations, we have to iterate over AllocationGroup sizes rather than TagData sizes.
+		// One extra bit of complexity: it is possible we are recording other TagSets as well, in which case
+		// each AllocationGroup tracks a more specific combination:
+		// { ELLMTagSet::None x ELLMTagSet::CodeOrContent x ELLMTagSet::<Other> }
+		// To handle that case, we have to sum up the size of every AllocationGroup that matches the 
+		// { ELLMTagSet::None x ELLMTagSet::CodeOrContent }
+		// combination we are interested in.
+		for (TPair<FColumnKey, FColumnData>& Pair : Columns)
 		{
-			continue;
+			Pair.Value.Size = 0;
 		}
-		if (TagData->GetName() == TagName_Untagged)
+
+		int32 IndexOfCodeOrContentInActiveTags = LLMRef.GetIndexInFActiveTags(ELLMTagSet::CodeOrContent);
+		// We are recording CodeOrContent, so it should have an index.
+		LLMCheck(IndexOfCodeOrContentInActiveTags >= 0);
+
+		UE::TSharedLock AllocationsReadLock(LLMRef.AllocationGroupsMutex);
+		for (const TPair<FActiveTags, FAllocationGroup*>& GroupPair : (*LLMRef.ActiveTagsToAllocationGroup))
 		{
-			continue; // Handled by OverrideUntaggedName.
+			const FActiveTags& ActiveTags = GroupPair.Key;
+			const FTagData* SystemsTag = ActiveTags.GetSystemsTagData();
+			if (!SystemsTag->IsReportable())
+			{
+				continue;
+			}
+			if (SystemsTag->GetDisplayName() == TagName_Untagged)
+			{
+				// We write the "Untagged" tag in the 2nd column, and always write it, rather than writing it in the
+				// column index where we first notice it. That writing is handled by use of OverrideUntaggedTagData in
+				// WriteHeader and AddRow.
+				continue; 
+			}
+			const FTagData* CodeOrContentTag = ActiveTags[IndexOfCodeOrContentInActiveTags];
+			if (!CodeOrContentTag)
+			{
+				continue;
+			}
+			int64 CurrentSize = GroupPair.Value->GetSize(this->Tracker);
+			// AllocationGroups includes some groups used by other Trackers and not by this one. We have to check whether
+			// the AllocationGroup has a non-zero size in our Tracker before adding it to our output. 
+			if (CurrentSize == 0)
+			{
+				continue;
+			}
+			FColumnData& ColumnData = Columns.FindOrAdd(FColumnKey{ SystemsTag, CodeOrContentTag });
+			if (!ColumnData.bInitialized)
+			{
+				bUpdated = true;
+				ColumnData.bInitialized = true;
+			}
+			ColumnData.Size += CurrentSize;
 		}
-		if (ExistingColumns.Contains(TagData))
+	}
+	else
+	{
+		for (const TPair<const FTagData*, FTrackerTagSizeData>& It : TagSizes)
 		{
-			continue;
-		}
+			const FTagData* SystemsTag = It.Key;
+			if (SystemsTag->GetTagSet() != ELLMTagSet::None || !SystemsTag->IsReportable())
+			{
+				continue;
+			}
+			if (SystemsTag->GetDisplayName() == TagName_Untagged)
+			{
+				// We write the "Untagged" tag in the 2nd column, and always write it, rather than writing it in the
+				// column index where we first notice it. That writing is handled by use of OverrideUntaggedTagData in
+				// WriteHeader and AddRow.
+				continue;
+			}
 
-		ExistingColumns.Add(TagData);
-		Columns.Add(TagData);
-		bUpdated = true;
+			FColumnData& ColumnData = Columns.FindOrAdd(FColumnKey{ SystemsTag, nullptr });
+			if (!ColumnData.bInitialized)
+			{
+				bUpdated = true;
+				ColumnData.bInitialized = true;
+			}
+			// ColumnData.Size is unused when !bRecordingCodeOrContent, because we already have
+			// the accumulated value in TagSizes and do not need to accumulate a sum.
+		}
 	}
 	return bUpdated;
 }
 
-void FLLMCsvWriter::WriteHeader(const FTagData* OverrideTrackedTotalTagData, const FTagData* OverrideUntaggedTagData)
+void FLLMCsvWriter::WriteHeader(FLLMGlobals& LLMRef, const FTagData* OverrideTrackedTotalTagData,
+	const FTagData* OverrideUntaggedTagData)
 {
 	int64 OriginalOffset = Archive->Tell();
 	Archive->Seek(0);
 
 	TStringBuilder<256> NameBuffer;
-	auto WriteTagData = [this, &NameBuffer](const FTagData* TagData)
+	auto WriteColumnName = [this, &NameBuffer](const FColumnKey& ColumnKey)
 	{
-		if (!TagData)
+		if (!ColumnKey.SystemsTag)
 		{
 			return;
 		}
 
 		NameBuffer.Reset();
-		TagData->AppendDisplayPath(NameBuffer);
+		ColumnKey.SystemsTag->AppendDisplayPath(NameBuffer);
+		if (ColumnKey.CodeOrContentTag)
+		{
+			NameBuffer << TEXT("/") << ColumnKey.CodeOrContentTag->GetDisplayName();
+		}
 		NameBuffer << TEXT(",");
 		Write(NameBuffer);
 	};
 
-	WriteTagData(OverrideTrackedTotalTagData);
-	WriteTagData(OverrideUntaggedTagData);
-	for (const FTagData* TagData : Columns)
+	// Specialcase in bRecordingCodeOrContent that makes it identical to the behavior of !bRecordingCodeOrContent:
+	// We combine the Code and Content values for ELLMTags::None tags TrackedTotal and Untagged into a single value
+	// and write a column for that, rather than making separate Code and Content values. It is not useful for budgets
+	// to show the Content-specific value of TrackedTotal and Untagged, so we merge the two columns to save space.
+	WriteColumnName(FColumnKey{ OverrideTrackedTotalTagData, nullptr });
+	WriteColumnName(FColumnKey{ OverrideUntaggedTagData, nullptr });
+	for (const TPair<FColumnKey, FColumnData>& Pair: Columns)
 	{
-		WriteTagData(TagData);
+		WriteColumnName(Pair.Key);
 	}
 
 	int64 ColumnTitleTotalSize = Archive->Tell();
@@ -7002,12 +7218,27 @@ void FLLMCsvWriter::AddRow(FLLMGlobals& LLMRef, const FTrackerTagSizeMap& TagSiz
 	}
 	if (OverrideUntaggedTagData)
 	{
-		WriteTag(LLMRef.FindTagData(TagName_Untagged, ELLMTagSet::None));
+		// All trackers store Untagged tagsizes under the "Untagged" tag.
+		// But for some not-yet-investigated legacy reason, when writing CSV we write out the name of a different ELLMTag:
+		// ELLMTag::UntaggedTotal or ELLMTag::PlatformUntaggedTotal.
+		// So do not read TagSizes[OverrideUntaggedTagData], read TagSizes[ELLMTag::Untagged].
+		const FTagData* TagToReadFromForUntagged = LLMRef.FindTagData(TagName_Untagged, ELLMTagSet::None);
+		WriteTag(TagToReadFromForUntagged);
 	}
 
-	for (const FTagData* TagData : Columns)
+	if (bRecordingCodeOrContent)
 	{
-		WriteTag(TagData);
+		for (TPair<FColumnKey, FColumnData>& Pair : Columns)
+		{
+			WriteValue(Pair.Value.Size);
+		}
+	}
+	else
+	{
+		for (TPair<FColumnKey, FColumnData>& Pair : Columns)
+		{
+			WriteTag(Pair.Key.SystemsTag);
+		}
 	}
 	Write(TEXTVIEW("\n"));
 
@@ -7088,15 +7319,17 @@ void FLLMTraceWriter::Publish(FLLMGlobals& LLMRef, const FTrackerTagSizeMap& Tag
 			"System",
 			"Assets",
 			"AssetClasses",
-			"UObjectClasses"
+			"UObjectClasses",
+			"CodeOrContent",
 		};
 		static_assert((int)ELLMTagSet::Assets == 1, "");
 		static_assert((int)ELLMTagSet::AssetClasses == 2, "");
 		static_assert((int)ELLMTagSet::UObjectClasses == 3, "");
+		static_assert((int)ELLMTagSet::CodeOrContent == 4, "");
 		static_assert(UE_ARRAY_COUNT(TagSetNames) == (int)ELLMTagSet::Max, "");
 		for (uint8 TagSet = 0; TagSet < (uint8)ELLMTagSet::Max; ++TagSet)
 		{
-			if (LLMRef.IsTagSetActive((ELLMTagSet)TagSet))
+			if (LLMRef.IsTagSetRecordingActiveInner((ELLMTagSet)TagSet))
 			{
 				FAnsiStringView TagSetName = TagSetNames[TagSet];
 				UE_TRACE_LOG(LLM, TagSetSpec, MemTagChannel, TagSetName.Len() * sizeof(ANSICHAR))
@@ -7126,7 +7359,7 @@ void FLLMTraceWriter::Publish(FLLMGlobals& LLMRef, const FTrackerTagSizeMap& Tag
 	for (const TPair<const FTagData*, FTrackerTagSizeData>& It : TagSizes)
 	{
 		const FTagData* TagData = It.Key;
-		if (!TagData->IsTraceReportable() || !LLMRef.IsTagSetActive(TagData->GetTagSet()))
+		if (!TagData->IsTraceReportable() || !LLMRef.IsTagSetRecordingActiveInner(TagData->GetTagSet()))
 		{
 			continue;
 		}
@@ -7171,7 +7404,7 @@ void FLLMTraceWriter::Publish(FLLMGlobals& LLMRef, const FTrackerTagSizeMap& Tag
 	for (const TPair<const FTagData*, FTrackerTagSizeData>& It : TagSizes)
 	{
 		const FTagData* TagData = It.Key;
-		if (!TagData->IsTraceReportable() || !LLMRef.IsTagSetActive(TagData->GetTagSet()))
+		if (!TagData->IsTraceReportable() || !LLMRef.IsTagSetRecordingActiveInner(TagData->GetTagSet()))
 		{
 			continue;
 		}
diff --git a/Engine/Source/Runtime/Core/Private/HAL/LLM/LLMPrivate.h b/Engine/Source/Runtime/Core/Private/HAL/LLM/LLMPrivate.h
--- a/Engine/Source/Runtime/Core/Private/HAL/LLM/LLMPrivate.h
+++ b/Engine/Source/Runtime/Core/Private/HAL/LLM/LLMPrivate.h
@@ -262,21 +262,33 @@ private:
 	void Write(FStringView Text);
 	static const TCHAR* GetTrackerCsvName(ELLMTracker InTracker);
 
-	bool CreateArchive();
-	bool UpdateColumns(const FTrackerTagSizeMap& TagSizes);
-	void WriteHeader(const FTagData* OverrideTrackedTotalTagData, const FTagData* OverrideUntaggedTagData);
+	bool CreateArchive(FLLMGlobals& LLMRef);
+	bool UpdateColumns(FLLMGlobals& LLMRef, const FTrackerTagSizeMap& TagSizes);
+	void WriteHeader(FLLMGlobals& LLMRef, const FTagData* OverrideTrackedTotalTagData, const FTagData* OverrideUntaggedTagData);
 	void AddRow(FLLMGlobals& LLMRef, const FTrackerTagSizeMap& TagSizes,
 		const FTagData* OverrideTrackedTotalName, const FTagData* OverrideUntaggedName, int64 TrackedTotal,
 		UE::LLM::ESizeParams SizeParams);
 
-	FConstTagDataArray Columns;
-	TLLMSet<const FTagData*> ExistingColumns;
+	struct FColumnKey
+	{
+		const FTagData* SystemsTag = nullptr;
+		const FTagData* CodeOrContentTag = nullptr;
+		bool operator==(const FColumnKey& Other) const;
+	};
+	friend uint32 GetTypeHash(const FColumnKey& Key);
+	struct FColumnData
+	{
+		int64 Size = 0;
+		bool bInitialized = false;
+	};
+	TLLMMap<FColumnKey, FColumnData> Columns;
 	FArchive* Archive;
 	double LastWriteTime;
 	int32 WriteCount;
 	int32 HeaderMaxSize;
 	ELLMTracker Tracker;
 	bool bRegisteredFlushDelegate;
+	bool bRecordingCodeOrContent;
 };
 
 /** Outputs the LLM tags and sizes to TraceLog events. */
@@ -558,6 +570,8 @@ public:
 	void SetProgramSizeInner(uint64 InProgramSize);
 	bool ExecInner(const TCHAR* Cmd, FOutputDevice& Ar);
 	bool IsTagSetActiveInner(ELLMTagSet Set);
+	bool IsTagSetScopeActiveInner(ELLMTagSet Set);
+	bool IsTagSetRecordingActiveInner(ELLMTagSet Set);
 	bool ShouldReduceThreadsInner();
 	const FTagData* GetActiveTagDataInner(ELLMTracker Tracker, ELLMTagSet TagSet = ELLMTagSet::None);
 	void RegisterPlatformTagInner(int32 Tag, const TCHAR* Name, FName StatName, FName SummaryStatName, int32 ParentTag = -1);
@@ -647,6 +661,8 @@ private:
 	const FTagData* FindTagData(FName Name, ELLMTagSet TagSet,
 		ETagReferenceSource ReferenceSource = ETagReferenceSource::FunctionAPI);
 	const FTagData* FindOrAddDefaultTagData(ELLMTagSet TagSet);
+	const FTagData* FindOrAddCodeOrContentCodeTagData();
+	const FTagData* FindOrAddCodeOrContentContentTagData();
 	FAllocationGroup* FindOrAddAllocationGroup(const UE::LLMPrivate::FActiveTags& ActiveTags);
 	FAllocationGroup* FindAllocationGroupForCompressedAllocInfo(ELLMTag Tag);
 
diff --git a/Engine/Source/Runtime/Core/Public/HAL/LowLevelMemStats.h b/Engine/Source/Runtime/Core/Public/HAL/LowLevelMemStats.h
--- a/Engine/Source/Runtime/Core/Public/HAL/LowLevelMemStats.h
+++ b/Engine/Source/Runtime/Core/Public/HAL/LowLevelMemStats.h
@@ -22,6 +22,7 @@
 	DECLARE_STATS_GROUP(TEXT("LLM Overhead"), STATGROUP_LLMOverhead, STATCAT_Advanced);
 	DECLARE_STATS_GROUP(TEXT("LLM Assets"), STATGROUP_LLMAssets, STATCAT_Advanced);
 	DECLARE_STATS_GROUP(TEXT("LLM UObjects"), STATGROUP_LLMUObjects, STATCAT_Advanced);
+	DECLARE_STATS_GROUP(TEXT("LLM CodeOrContent"), STATGROUP_LLMCodeOrContent, STATCAT_Advanced);
 
 	/*
 	* LLM Summary stats referenced by ELLMTagNames
diff --git a/Engine/Source/Runtime/Core/Public/HAL/LowLevelMemTracker.h b/Engine/Source/Runtime/Core/Public/HAL/LowLevelMemTracker.h
--- a/Engine/Source/Runtime/Core/Public/HAL/LowLevelMemTracker.h
+++ b/Engine/Source/Runtime/Core/Public/HAL/LowLevelMemTracker.h
@@ -124,9 +124,20 @@ enum class ELLMTracker : uint8
 enum class ELLMTagSet : uint8
 {
 	None,
+	/** Create tag for each Asset, and assign allocations for any UObject in a package to the package's asset. */
 	Assets,
+	/**
+	 * Create tag for each UObject class, and assign allocations for any UObject in a package to the class of the
+	 * package's asset.
+	 */
 	AssetClasses,
+	/** Create tag for each UObject class, and assign allocations for any UObject to the class of the UObject. */
 	UObjectClasses,
+	/**
+	 * Create two tags - Code or Content. Assign allocations for any UObject to the Content tag; assign allocations
+	 * shared by multiple objects or unrelated to UObjects to the Code tag.
+	 */
+	CodeOrContent,
 	
 	Max, // note: update functions that handle ELLMTagSets if you add any new tagsets
 };
@@ -404,16 +415,23 @@ extern CORE_API const FName LLMGetUntaggedTagName(ELLMTagSet TagSet);
 
 /**
  * Define a tag which can be used in LLM_SCOPE_BYTAG or referenced by name in other LLM_SCOPEs.
- * @param UniqueNameWithUnderscores - Modified version of the name of the tag. Used for looking up by name,
+ * @param UniqueNameWithUnderscores - Token useable as c++ variable name.
+ *        Modified version of the name of the tag. Used for looking up by name,
  *        must be unique across all tags passed to LLM_DEFINE_TAG, LLM_SCOPE, or ELLMTag.
  *        The modification: the usual separator / for parents must be replaced with _ in LLM_DEFINE_TAGs.
- * @param DisplayName - (Optional) - The name to display when tracing the tag; joined with "/" to the name of its
+ * @param DisplayName - (Optional) - FName
+*         The name to display when tracing the tag; joined with "/" to the name of its
  *        parent if it has a parent, or NAME_None to use the UniqueName.
- * @param ParentTagName - (Optional) - The unique name of the parent tag, or NAME_None if it has no parent.
- * @param StatName - (Optional) - The name of the stat to populate with this tag's amount when publishing LLM data each
+ * @param ParentTagName - (Optional) - FName
+ *        The unique name of the parent tag, or NAME_None if it has no parent.
+ * @param StatName - (Optional) - FName
+ *        The name of the stat to populate with this tag's amount when publishing LLM data each
  *        frame, or NAME_None if no stat should be populated.
- * @param SummaryStatName - (Optional) - The name of the stat group to add on this tag's amount when publishing LLM
+ * @param SummaryStatName - (Optional) - FName
+ *        The name of the stat group to add on this tag's amount when publishing LLM
  *        data each frame, or NAME_None if no stat group should be added to.
+ * @param TagSet - (Optional) - ELLMTagSet
+ *        The TagSet the FName is assigned to. Default is ELLMTagSet::None.
  */
 #define LLM_DEFINE_TAG(UniqueNameWithUnderscores, ...) FLLMTagDeclaration UE_JOIN(LLMTagDeclaration_, UniqueNameWithUnderscores)(TEXT(#UniqueNameWithUnderscores), ##__VA_ARGS__)
 
@@ -749,6 +767,7 @@ protected:
 	// All fields other than bEnabled are not initialized by constructor unless bEnabled is true
 	ELLMTracker Tracker;
 	bool bEnabled = false;  // Needs to be a uint8 rather than a bitfield to minimize ALU operations
+	bool bPushedCodeOrContent = false;
 	ELLMTagSet TagSet;
 };
 
@@ -814,6 +833,7 @@ protected:
 
 	// All fields other than bEnabled are not initialized by constructor unless bEnabled is true
 	const UE::LLMPrivate::FTagData* TagData;
+	const UE::LLMPrivate::FTagData* CodeOrContentTagData;
 	ELLMTracker Tracker;
 	bool bEnabled = false; // Needs to be a uint8 rather than a bitfield to minimize ALU operations
 	ELLMTagSet TagSet;
