Tmp fix for content check failure: remove check and replace with defensive if until root cause found

#Synced-CL 49767845

diff --git a/Engine/Source/Runtime/Core/Private/Features/ModularFeatures.cpp b/Engine/Source/Runtime/Core/Private/Features/ModularFeatures.cpp
--- a/Engine/Source/Runtime/Core/Private/Features/ModularFeatures.cpp
+++ b/Engine/Source/Runtime/Core/Private/Features/ModularFeatures.cpp
@@ -66,7 +66,11 @@ void FModularFeatures::RegisterModularFeature( const FName Type, IModularFeature
 
 void FModularFeatures::UnregisterModularFeature( const FName Type, IModularFeature* ModularFeature )
 {
-	check(ModularFeature != nullptr);
+	if(ModularFeature == nullptr)
+	{
+		return;
+	}
+
 	check(Type.IsValid());
 
 	UE::TScopeLock ScopeLock(ModularFeaturesMapCriticalSection);
