Disable ProgramSize metric on macOS to match Windows behavior

#Synced-CL 49887699

diff --git a/Engine/Source/Runtime/Core/Private/Apple/ApplePlatformMemory.cpp b/Engine/Source/Runtime/Core/Private/Apple/ApplePlatformMemory.cpp
--- a/Engine/Source/Runtime/Core/Private/Apple/ApplePlatformMemory.cpp
+++ b/Engine/Source/Runtime/Core/Private/Apple/ApplePlatformMemory.cpp
@@ -202,6 +202,16 @@ void FApplePlatformMemory::Init()
     
 	LLM(AppleLLM::Initialise());
 
+#if PLATFORM_MAC
+	// Following Windows' lead on disabling program size. The default implementation
+	// of program size (in FLLMGlobals::InitializeProgramSize) right now is 
+	// FPlatformMemoryStats::TotalPhysical - FPlatformMemoryStats::AvailablePhysical.
+	// Unless we are running in an environment primarily designed around one process,
+	// this is incorrect as both of these numbers are system-level stats and could include
+	// memory taken up by other processes.
+	LLM(FLowLevelMemTracker::Get().SetProgramSize(0));
+#endif
+	
 	const FPlatformMemoryConstants& MemoryConstants = FPlatformMemory::GetConstants();
 	UE_LOG(LogInit, Log, TEXT("Memory total: Physical=%.1fGB (%dGB approx) Pagefile=%.1fGB Virtual=%.1fGB"),
 		   float(MemoryConstants.TotalPhysical/1024.0/1024.0/1024.0),
