PagedArray: Modify NumPages() to return the number of pages required for holding the current number of elements rather than the number of currently allocated pages.
The function has been unused outside of PagedArrayTest.cpp in engine and project code.
Add MaxPages in case some client code needs the old behavior.

#rnx

#Synced-CL 49808502

diff --git a/Engine/Source/Runtime/Core/Public/Containers/Array.h b/Engine/Source/Runtime/Core/Public/Containers/Array.h
--- a/Engine/Source/Runtime/Core/Public/Containers/Array.h
+++ b/Engine/Source/Runtime/Core/Public/Containers/Array.h
@@ -1259,9 +1259,9 @@ public:
 	}
 
 	/**
-	 * Returns maximum number of elements in array.
+	 * Returns maximum number of elements in array before reallocation.
 	 *
-	 * @returns Maximum number of elements in array.
+	 * @returns Maximum number of elements in array before reallocation.
 	 * @see GetSlack
 	 */
 	[[nodiscard]] UE_REWRITE SizeType Max() const
diff --git a/Engine/Source/Runtime/Core/Public/Containers/PagedArray.h b/Engine/Source/Runtime/Core/Public/Containers/PagedArray.h
--- a/Engine/Source/Runtime/Core/Public/Containers/PagedArray.h
+++ b/Engine/Source/Runtime/Core/Public/Containers/PagedArray.h
@@ -282,21 +282,48 @@ public:
 		Ar.CountBytes(Num() * sizeof(ElementType), Max() * sizeof(ElementType));
 	}
 
+	/**
+	 * Returns maximum number of elements in array before reallocation.
+	 *
+	 * @returns Maximum number of elements in array before reallocation.
+	 */
 	[[nodiscard]] UE_FORCEINLINE_HINT SizeType Max() const
 	{
 		return Pages.Num() * PageTraits::Capacity;
 	}
 
+	/**
+	 * Returns number of elements in array.
+	 *
+	 * @returns Number of elements in array.
+	 */
 	[[nodiscard]] UE_FORCEINLINE_HINT SizeType Num() const
 	{
 		return Count;
 	}
 
-	[[nodiscard]] UE_FORCEINLINE_HINT SizeType NumPages() const
+	/**
+	 * Returns maximum number of pages before reallocation.
+	 *
+	 * @returns Returns maximum number of pages before reallocation.
+	 */
+	[[nodiscard]] UE_FORCEINLINE_HINT SizeType MaxPages() const
 	{
 		return Pages.Num();
 	}
 
+	/**
+	 * Returns number of pages used for the number of elements in array.
+	 *
+	 * @returns Returns number of pages used for the number of elements in array.
+	 */
+	[[nodiscard]] UE_FORCEINLINE_HINT SizeType NumPages() const
+	{
+		SizeType NumPages = NumRequiredPages(Count);
+		checkSlow(static_cast<int32>(NumPages) <= Pages.Num());
+		return NumPages;
+	}
+
 	[[nodiscard]] UE_FORCEINLINE_HINT bool IsValidIndex(SizeType Index) const
 	{
 		return Index >= 0 && Index < Count;
@@ -826,4 +853,4 @@ template <typename ElementType, int32 PageSizeInBytes, typename AllocatorType>
 UE_FORCEINLINE_HINT void operator<<(FStructuredArchive::FSlot Slot, TPagedArray<ElementType, PageSizeInBytes, AllocatorType>& InOutPagedArray)
 {
 	UE::Core::PagedArray::Private::SerializeStructured(Slot, InOutPagedArray);
-}
\ No newline at end of file
+}
diff --git a/Engine/Source/Runtime/Core/Tests/Containers/PagedArrayTest.cpp b/Engine/Source/Runtime/Core/Tests/Containers/PagedArrayTest.cpp
--- a/Engine/Source/Runtime/Core/Tests/Containers/PagedArrayTest.cpp
+++ b/Engine/Source/Runtime/Core/Tests/Containers/PagedArrayTest.cpp
@@ -40,10 +40,14 @@ TEST_CASE_NAMED(FPagedArrayReserveWithoutDataTest, "System::Core::Containers::TP
 	ArrayType PagedArray;
 	CHECK(!PagedArray.Max());
 	CHECK(!PagedArray.Num());
+	CHECK(!PagedArray.MaxPages());
+	CHECK(!PagedArray.NumPages());
 	CHECK(PagedArray.IsEmpty());
 	PagedArray.Reserve(ArrayType::MaxPerPage());
 	CHECK(PagedArray.Max() == ArrayType::MaxPerPage());
 	CHECK(!PagedArray.Num());
+	CHECK(PagedArray.MaxPages() == 1);
+	CHECK(!PagedArray.NumPages());
 	CHECK(PagedArray.IsEmpty());
 	CHECK(int32Token::EvenConstructionDestructionCalls(0));
 }
@@ -66,14 +70,20 @@ TEST_CASE_NAMED(
 		CHECK(PagedArray.Max() >= 1);
 		CHECK(PagedArray.Max() == ArrayType::MaxPerPage());
 		CHECK(PagedArray.Num() == 1);
+		CHECK(PagedArray.MaxPages() == 1);
+		CHECK(PagedArray.NumPages() == 1);
 
 		PagedArray.Reserve(ArrayType::MaxPerPage());
 		CHECK(PagedArray.Max() == ArrayType::MaxPerPage());
 		CHECK(PagedArray.Num() == 1);
+		CHECK(PagedArray.MaxPages() == 1);
+		CHECK(PagedArray.NumPages() == 1);
 
 		PagedArray.Reserve(1 + ArrayType::MaxPerPage());
 		CHECK(PagedArray.Max() == ArrayType::MaxPerPage() * 2);
 		CHECK(PagedArray.Num() == 1);
+		CHECK(PagedArray.MaxPages() == 2);
+		CHECK(PagedArray.NumPages() == 1);
 	}
 	CHECK(int32Token::EvenConstructionDestructionCalls(1));
 }
@@ -613,6 +623,8 @@ TEST_CASE_NAMED(FPagedArraySetNumAfterResetTest, "System::Core::Containers::TPag
 
 		// Reset and set same number of elements
 		PagedArray.Reset();
+		CHECK(PagedArray.MaxPages() == 3);
+		CHECK(!PagedArray.NumPages());
 		PagedArray.SetNum(ArrayType::MaxPerPage() * 3);
 		CHECK(PagedArray.Num() == ArrayType::MaxPerPage() * 3);
 		CHECK(PagedArray.NumPages() == 3);
@@ -631,6 +643,8 @@ TEST_CASE_NAMED(FPagedArraySetNumAfterResetTest, "System::Core::Containers::TPag
 
 		// Reset and set a greater number of elements
 		PagedArray.Reset();
+		CHECK(PagedArray.MaxPages() == 3);
+		CHECK(!PagedArray.NumPages());
 		PagedArray.SetNum(ArrayType::MaxPerPage() * 5);
 		CHECK(PagedArray.Num() == ArrayType::MaxPerPage() * 5);
 		CHECK(PagedArray.NumPages() == 5);
