Improve performance of FUnixPlatformProcess::ExecProcess when used to launch a process with a lot of stdout/stderror traffic.

- The code was waiting for 0.5 seconds after each read from stdout+stderr. Not only is this very slow but the long waits could cause the pipes to run out of buffer and cause stalls on the launched process as well.
-- As an example, one fairly spammy process was seen to take ~ 70 seconds to run on unix vs ~7 seconds with this fix.
- Waiting for 0 seconds in every loop (to yield) should now produce similar performance to our ::FWindowsPlatformProcess implementation.

#Synced-CL 49815329

diff --git a/Engine/Source/Runtime/Core/Private/Unix/UnixPlatformProcess.cpp b/Engine/Source/Runtime/Core/Private/Unix/UnixPlatformProcess.cpp
--- a/Engine/Source/Runtime/Core/Private/Unix/UnixPlatformProcess.cpp
+++ b/Engine/Source/Runtime/Core/Private/Unix/UnixPlatformProcess.cpp
@@ -2106,7 +2106,7 @@ bool FUnixPlatformProcess::ExecProcess(const TCHAR* URL, const TCHAR* Params, in
 		{
 			ReadPipeToStr(PipeReadStdOut, OutStdOut);
 			ReadPipeToStr(PipeReadStdErr, OutStdErr);
-			FPlatformProcess::Sleep(0.5);
+			FPlatformProcess::Sleep(0.0);
 		}
 
 		// Read the remainder
@@ -2125,8 +2125,6 @@ bool FUnixPlatformProcess::ExecProcess(const TCHAR* URL, const TCHAR* Params, in
 			}
 		}
 
-		FPlatformProcess::Sleep(0.5);
-
 		bInvoked = true;
 		bool bGotReturnCode = FPlatformProcess::GetProcReturnCode(ProcHandle, &ReturnCode);
 		check(bGotReturnCode)
