[IOS] IOSPlatformMisc.GetDeviceThermalState was missing WITH_APPLICATION_CORE.
#rnx

#Synced-CL 49802443

diff --git a/Engine/Source/Runtime/Core/Private/IOS/IOSPlatformMisc.cpp b/Engine/Source/Runtime/Core/Private/IOS/IOSPlatformMisc.cpp
--- a/Engine/Source/Runtime/Core/Private/IOS/IOSPlatformMisc.cpp
+++ b/Engine/Source/Runtime/Core/Private/IOS/IOSPlatformMisc.cpp
@@ -470,7 +470,7 @@ float FIOSPlatformMisc::GetDeviceTemperatureLevel()
 
 EDeviceThermalState FIOSPlatformMisc::GetDeviceThermalState()
 {
-#if !PLATFORM_TVOS
+#if WITH_APPLICATION_CORE && !PLATFORM_TVOS
 	// We are mapping the ios thermal state values onto the first four of our thermal state enum values, which are based on 
 	// the android thermal status values.
 	// Nominal and None both refer to states where the device is entirely fine on heat and not throttling.
