Added IsLocked to FWordMutex\FRecursiveWordMutex

#Synced-CL 49805908

diff --git a/Engine/Source/Runtime/Core/Public/Async/RecursiveWordMutex.h b/Engine/Source/Runtime/Core/Public/Async/RecursiveWordMutex.h
--- a/Engine/Source/Runtime/Core/Public/Async/RecursiveWordMutex.h
+++ b/Engine/Source/Runtime/Core/Public/Async/RecursiveWordMutex.h
@@ -22,6 +22,10 @@ namespace UE
 class FRecursiveWordMutex final
 {
 public:
+	[[nodiscard]] inline bool IsLocked() const
+	{
+		return Mutex.IsLocked();
+	}
 	[[nodiscard]] UE_API bool TryLock();
 	UE_API void Lock();
 	UE_API void Unlock();
diff --git a/Engine/Source/Runtime/Core/Public/Async/WordMutex.h b/Engine/Source/Runtime/Core/Public/Async/WordMutex.h
--- a/Engine/Source/Runtime/Core/Public/Async/WordMutex.h
+++ b/Engine/Source/Runtime/Core/Public/Async/WordMutex.h
@@ -25,6 +25,11 @@ public:
 	FWordMutex(const FWordMutex&) = delete;
 	FWordMutex& operator=(const FWordMutex&) = delete;
 
+	[[nodiscard]] inline bool IsLocked() const
+	{
+		return !!(State.load(std::memory_order_relaxed) & IsLockedFlag);
+	}
+
 	[[nodiscard]] inline bool TryLock()
 	{
 		UPTRINT Expected = 0;
