[Android] GetLastError implementation

This was missing on Android, so generic one was used, which always returned 0.

#rnx

diff --git a/Engine/Source/Runtime/Core/Private/Android/AndroidPlatformMisc.cpp b/Engine/Source/Runtime/Core/Private/Android/AndroidPlatformMisc.cpp
--- a/Engine/Source/Runtime/Core/Private/Android/AndroidPlatformMisc.cpp
+++ b/Engine/Source/Runtime/Core/Private/Android/AndroidPlatformMisc.cpp
@@ -736,6 +736,11 @@ const TCHAR* FAndroidMisc::GetSystemErrorMessage(TCHAR* OutBuffer, int32 BufferC
 	return OutBuffer;
 }
 
+uint32 FAndroidMisc::GetLastError()
+{
+	return (uint32)errno;
+}
+
 EAppReturnType::Type FAndroidMisc::MessageBoxExt( EAppMsgType::Type MsgType, const TCHAR* Text, const TCHAR* Caption )
 {
 #if USE_ANDROID_JNI
diff --git a/Engine/Source/Runtime/Core/Public/Android/AndroidPlatformMisc.h b/Engine/Source/Runtime/Core/Public/Android/AndroidPlatformMisc.h
--- a/Engine/Source/Runtime/Core/Public/Android/AndroidPlatformMisc.h
+++ b/Engine/Source/Runtime/Core/Public/Android/AndroidPlatformMisc.h
@@ -58,6 +58,7 @@ struct FAndroidMisc : public FGenericPlatformMisc
 
 	static CORE_API FString GetEnvironmentVariable(const TCHAR* VariableName);
 	static CORE_API const TCHAR* GetSystemErrorMessage(TCHAR* OutBuffer, int32 BufferCount, int32 Error);
+	static CORE_API uint32 GetLastError();
 	static CORE_API EAppReturnType::Type MessageBoxExt( EAppMsgType::Type MsgType, const TCHAR* Text, const TCHAR* Caption );
 	static CORE_API bool UseRenderThread();
 	static CORE_API bool HasPlatformFeature(const TCHAR* FeatureName);
