Fixed FContainersTSetTest error. Mismatching FRecorder::operator== and FRecorder::GetTypeHash.

#Synced-CL 49850619

diff --git a/Engine/Source/Runtime/Core/Private/Containers/ContainersTest.cpp b/Engine/Source/Runtime/Core/Private/Containers/ContainersTest.cpp
--- a/Engine/Source/Runtime/Core/Private/Containers/ContainersTest.cpp
+++ b/Engine/Source/Runtime/Core/Private/Containers/ContainersTest.cpp
@@ -1274,7 +1274,7 @@ bool FContainerPerformanceTest::RunTest(const FString& Parameters)
 	return true;
 }
 
-namespace
+namespace ContainersTSetTest
 {
 	struct FRecorder
 	{
@@ -1318,13 +1318,15 @@ namespace
 
 	uint32 GetTypeHash(const FRecorder& Recorder)
 	{
-		return Recorder.Payload;
+		return Recorder.Key;
 	}
 
 }
 IMPLEMENT_SIMPLE_AUTOMATION_TEST(FContainersTSetTest, "System.Core.Containers.TSet", EAutomationTestFlags_ApplicationContextMask | EAutomationTestFlags::EngineFilter)
 bool FContainersTSetTest::RunTest(const FString& Parameters)
 {
+	using namespace ContainersTSetTest;
+
 	enum class EArgType
 	{
 		Copy,
