[Insights] Secure socket support in Trace auxiliary

- Added support for secure trace connections.
- Added delegate for OnClose event, when a connection is closed.
- Added support for logging TraceLog messages with Verbose and VeryVerbose levels.

#Synced-CL 49918891

diff --git a/Engine/Source/Runtime/Core/Private/ProfilingDebugging/TraceAuxiliary.cpp b/Engine/Source/Runtime/Core/Private/ProfilingDebugging/TraceAuxiliary.cpp
--- a/Engine/Source/Runtime/Core/Private/ProfilingDebugging/TraceAuxiliary.cpp
+++ b/Engine/Source/Runtime/Core/Private/ProfilingDebugging/TraceAuxiliary.cpp
@@ -3,6 +3,7 @@
 #include "ProfilingDebugging/TraceAuxiliary.h"
 
 #include "BuildSettings.h"
+#include "Containers/AnsiString.h"
 #include "ProfilingDebugging/CountersTrace.h"
 #include "ProfilingDebugging/CsvProfilerConfig.h"
 #include "ProfilingDebugging/StringsTrace.h"
@@ -154,7 +155,14 @@ public:
 	void DisableChannels(TConstArrayView<uint32> ChannelIds, TMap<uint32, FString>* OutReason);
 	bool Connect(FTraceAuxiliary::EConnectionType Type, const TCHAR* Parameter, const FTraceAuxiliary::FLogCategoryAlias& LogCategory, uint16 SendFlags);
 	bool ConnectRelay(UPTRINT Handle, UE::Trace::IoWriteFunc WriteFunc, UE::Trace::IoCloseFunc CloseFunc, uint16 SendFlags);
+#if UE_TRACE_ALLOW_SECURE_TRACING
+	bool InitSecure();
+	bool ConnectSecure(FTraceAuxiliary::FSecureOptions& Options, uint16 SendFlags);
+	FString GetSecureHost() const;
+	uint16 GetSecurePort() const;
+#endif
 	bool Stop();
+	void OnClosed();
 	void FreezeReadOnlyChannels();
 	void ResumeChannels();
 	void PauseChannels();
@@ -215,6 +223,15 @@ private:
 	FRWLock DiagnosticsLock;
 	FString DiagnosticMessage;
 
+#if UE_TRACE_ALLOW_SECURE_TRACING
+	mutable FRWLock SecureAuthLock;
+	TArray<uint8> SecureToken;
+	bool bSecureInit = false;
+	FAnsiString SecureIdentity;
+	FString SecureHost;
+	uint16 SecurePort = 0;
+#endif
+
 	struct FCurrentTraceTarget
 	{
 		FString TraceDest;
@@ -608,6 +625,8 @@ bool FTraceAuxiliaryImpl::Connect(FTraceAuxiliary::EConnectionType Type, const T
 ////////////////////////////////////////////////////////////////////////////////
 bool FTraceAuxiliaryImpl::ConnectRelay(UPTRINT Handle, UE::Trace::IoWriteFunc WriteFunc, UE::Trace::IoCloseFunc CloseFunc, uint16 SendFlags)
 {
+	check(WriteFunc && CloseFunc);
+
 	if (UE::Trace::IsTracing())
 	{
 		UE_LOG(LogTrace, Warning, TEXT("Already tracing!"));
@@ -638,6 +657,136 @@ bool FTraceAuxiliaryImpl::ConnectRelay(UPTRINT Handle, UE::Trace::IoWriteFunc Wr
 	return bConnected;
 }
 
+////////////////////////////////////////////////////////////////////////////////
+#if UE_TRACE_ALLOW_SECURE_TRACING 
+
+////////////////////////////////////////////////////////////////////////////////
+bool FTraceAuxiliaryImpl::InitSecure()
+{
+	if (!bSecureInit)
+	{
+		FString CAPath, PrivateKeyPath;
+		int32 SecurePortInt32 = 0;
+		bool bUsePreSharedKeys = false;
+		if (GConfig)
+		{
+			GConfig->GetString(TEXT("Trace"), TEXT("SecureCA"), CAPath, GEngineIni);
+			GConfig->GetString(TEXT("Trace"), TEXT("SecureKey"), PrivateKeyPath, GEngineIni);
+			GConfig->GetBool(TEXT("Trace"), TEXT("SecureUsePreSharedKeys"), bUsePreSharedKeys, GEngineIni);
+			GConfig->GetInt(TEXT("Trace"), TEXT("SecurePort"), SecurePortInt32, GEngineIni);
+		}
+
+		FParse::Value(FCommandLine::Get(), TEXT("-TraceSecureCA="), CAPath);
+		FParse::Value(FCommandLine::Get(), TEXT("-TraceSecureKey="), PrivateKeyPath);
+		FParse::Value(FCommandLine::Get(), TEXT("-TraceSecureHost="), SecureHost);
+		FParse::Value(FCommandLine::Get(), TEXT("-TraceSecurePort="), SecurePortInt32);
+		bUsePreSharedKeys |= FParse::Param(FCommandLine::Get(), TEXT("-TraceSecureUsePreSharedKeys"));
+		FAnsiString CAPathAnsi(CAPath);
+		FAnsiString PrivateKeyPathAnsi(PrivateKeyPath);
+
+		UE_LOG(LogTrace, Verbose, TEXT("Initializing secure trace."));
+		UE::Trace::FSecureTraceSettings Settings;
+
+		// Setup how to connect to outside peers
+		if (SecureHost.IsEmpty())
+		{
+			UE_LOG(LogTrace, Display, TEXT("No secure host specified, public IP will be used."));
+		}
+		if (SecurePortInt32 > 0 && SecurePortInt32 < UINT16_MAX)
+		{
+			Settings.Port = SecurePort = uint16(SecurePortInt32);
+		}
+		else
+		{
+			UE_LOG(LogTrace, Display, TEXT("No secure port specified, system will assign port."));
+			Settings.Port = 0;
+		}
+
+		// Determine which security mode we will use.
+		if (!CAPathAnsi.IsEmpty() && !PrivateKeyPathAnsi.IsEmpty())
+		{
+			//todo: Verify files exists and are readable
+			Settings.Mode = UE::Trace::ESecureTraceMode::CertificateFile;
+			UE_LOG(LogTrace, Verbose, TEXT("Using certificate '%hs' and key '%hs'"), *CAPathAnsi, *PrivateKeyPathAnsi);
+			Settings.CertificatePath = *CAPathAnsi;
+			Settings.PrivateKeyPath = *PrivateKeyPathAnsi;
+		}
+		else if (bUsePreSharedKeys)
+		{
+			Settings.Mode = UE::Trace::ESecureTraceMode::PreSharedKeys;
+			UE_LOG(LogTrace, Verbose, TEXT("Secure trace mode: Using pre shared keys"));
+		}
+		else
+		{
+			Settings.Mode = UE::Trace::ESecureTraceMode::SelfSigned;
+			UE_LOG(LogTrace, Verbose, TEXT("Secure trace mode: Using self signed certificate"));
+		}
+
+		if (!UE::Trace::SecureInitialize(Settings))
+		{
+			UE_LOG(LogTrace, Error, TEXT("Secure trace initialization failed."));
+			return false;
+		}
+		bSecureInit = true;
+	}
+	return true;
+}
+
+////////////////////////////////////////////////////////////////////////////////
+bool FTraceAuxiliaryImpl::ConnectSecure(FTraceAuxiliary::FSecureOptions& InOptions, uint16 SendFlags)
+{
+	UE::Trace::FSecureTraceOptions SecureOptions;
+	if (!InOptions.AuthorizationToken.IsEmpty())
+	{
+		// We're okay to write this before the actual connection
+		// is created here since we don't allow overlapping ongoing
+		// connections with a new pending connection.
+		FWriteScopeLock _(SecureAuthLock);
+		if (!InitSecure())
+		{
+			return false;
+		}
+		SecureToken = InOptions.AuthorizationToken;
+		SecureIdentity = InOptions.Identity;
+		SecureOptions.AuthorizationTokenSize = SecureToken.Num();
+		SecureOptions.AuthorizationToken = SecureToken.GetData();
+		SecureOptions.Identity = *SecureIdentity;
+	}
+	else
+	{
+		SecureOptions.AuthorizationTokenSize = 0u;
+		SecureOptions.AuthorizationToken = nullptr;
+		SecureOptions.Identity = nullptr;
+	}
+	SecureOptions.TimeoutMs = uint32(InOptions.TimeoutSeconds * 1000);
+
+	const bool bSuccess = UE::Trace::SecureSendTo(SecureOptions, SendFlags);
+	if (bSuccess)
+	{
+		FWriteScopeLock _(CurrentTargetLock);
+		CurrentTraceTarget.TraceType = FTraceAuxiliary::EConnectionType::SecureNetwork;
+		CurrentTraceTarget.TraceDest = TEXT("Secure network trace");
+	}
+
+	return bSuccess;
+}
+
+////////////////////////////////////////////////////////////////////////////////
+FString FTraceAuxiliaryImpl::GetSecureHost() const
+{
+	FWriteScopeLock _(SecureAuthLock);
+	return SecureHost;
+}
+
+////////////////////////////////////////////////////////////////////////////////
+uint16 FTraceAuxiliaryImpl::GetSecurePort() const
+{
+	FWriteScopeLock _(SecureAuthLock);
+	return SecurePort;
+}
+
+#endif // UE_TRACE_ALLOW_SECURE_TRACING
+
 ////////////////////////////////////////////////////////////////////////////////
 bool FTraceAuxiliaryImpl::Stop()
 {
@@ -653,24 +802,41 @@ bool FTraceAuxiliaryImpl::Stop()
 		return false;
 	}
 
-	FString StopedDest;
-	FTraceAuxiliary::EConnectionType StopedType = FTraceAuxiliary::EConnectionType::None;
+	FString StoppedDest;
+	FTraceAuxiliary::EConnectionType StoppedType = FTraceAuxiliary::EConnectionType::None;
 	PausedPreset.Empty();
 
 	{
 		FWriteScopeLock _(CurrentTargetLock);
-		StopedType = CurrentTraceTarget.TraceType;
+		StoppedType = CurrentTraceTarget.TraceType;
 		CurrentTraceTarget.TraceType = FTraceAuxiliary::EConnectionType::None;
-		StopedDest = CurrentTraceTarget.TraceDest;
+		StoppedDest = CurrentTraceTarget.TraceDest;
 		CurrentTraceTarget.TraceDest.Reset();
 	}
 
 	TRACE_CPUPROFILER_EVENT_SCOPE(FTraceAux_OnTraceStopped);
-	FTraceAuxiliary::OnTraceStopped.Broadcast(StopedType, StopedDest);
+	FTraceAuxiliary::OnTraceStopped.Broadcast(StoppedType, StoppedDest);
 
 	return true;
 }
 
+////////////////////////////////////////////////////////////////////////////////
+void FTraceAuxiliaryImpl::OnClosed()
+{
+	// Clear the current target values
+	{
+		FWriteScopeLock _(CurrentTargetLock);
+		CurrentTraceTarget.TraceType = FTraceAuxiliary::EConnectionType::None;
+		CurrentTraceTarget.TraceDest.Reset();
+	}
+#if UE_TRACE_ALLOW_SECURE_TRACING
+	{
+		FWriteScopeLock _(SecureAuthLock);
+		SecureToken.Reset(0);
+	}
+#endif
+}
+
 ////////////////////////////////////////////////////////////////////////////////
 void FTraceAuxiliaryImpl::FreezeReadOnlyChannels()
 {
@@ -1188,6 +1354,14 @@ void TraceAuxiliaryOnConnectionCallback()
 	FTraceAuxiliary::OnConnection.Broadcast();
 }
 
+////////////////////////////////////////////////////////////////////////////////
+void TraceAuxiliaryOnCloseCallback() 
+{
+	TRACE_CPUPROFILER_EVENT_SCOPE(FTraceAux_OnClosed);
+	GTraceAuxiliary.OnClosed();
+	FTraceAuxiliary::OnClosed.Broadcast();
+}
+
 ////////////////////////////////////////////////////////////////////////////////
 void TraceAuxiliaryOnMessageCallback(const UE::Trace::FMessageEvent& Message)
 {
@@ -1203,6 +1377,12 @@ void TraceAuxiliaryOnMessageCallback(const UE::Trace::FMessageEvent& Message)
 
 	switch (Message.Type)
 	{
+	case UE::Trace::EMessageType::VeryVerbose:
+		LOG_WITH_VERBOSITY(VeryVerbose, Msg);
+		break;
+	case UE::Trace::EMessageType::Verbose:
+		LOG_WITH_VERBOSITY(Verbose, Msg);
+		break;
 	case UE::Trace::EMessageType::Log:
 		LOG_WITH_VERBOSITY(Log, Msg);
 		break;
@@ -1776,6 +1956,7 @@ static bool StartFromCommandlineArguments(const TCHAR* CommandLine)
 
 ////////////////////////////////////////////////////////////////////////////////
 FTraceAuxiliary::FOnConnection FTraceAuxiliary::OnConnection;
+FTraceAuxiliary::FOnClosed FTraceAuxiliary::OnClosed;
 FTraceAuxiliary::FOnTraceStarted FTraceAuxiliary::OnTraceStarted;
 FTraceAuxiliary::FOnTraceStopped FTraceAuxiliary::OnTraceStopped;
 FTraceAuxiliary::FOnSnapshotSaved FTraceAuxiliary::OnSnapshotSaved;
@@ -1860,6 +2041,53 @@ bool FTraceAuxiliary::Relay(UPTRINT Handle, UE::Trace::IoWriteFunc WriteFunc, UE
 #endif
 }
 
+#if UE_TRACE_ALLOW_SECURE_TRACING
+
+////////////////////////////////////////////////////////////////////////////////
+bool FTraceAuxiliary::StartSecure(FSecureOptions& SecureOptions, FOptions* Options, const TCHAR* Channels)
+{
+#if UE_TRACE_MINIMAL_ENABLED
+	TRACE_CPUPROFILER_EVENT_SCOPE(FTraceAux_Secure);
+
+	if (!TraceAuxiliaryStartShared(Channels, Options, LogTrace))
+	{
+		return false;
+	}
+
+	uint16 SendFlags = (Options && Options->bExcludeTail) ? UE::Trace::FSendFlags::ExcludeTail : 0;
+
+	return GTraceAuxiliary.ConnectSecure(SecureOptions, SendFlags);
+#else // UE_TRACE_MINIMAL_ENABLED
+	return false;
+#endif
+}
+
+////////////////////////////////////////////////////////////////////////////////
+TArrayView<uint8> FTraceAuxiliary::GetSecureCertificate()
+{
+	uint8* Ptr;
+	uint32 Size;
+	if (!UE::Trace::SecureGetSelfSignedCertificate(&Ptr, &Size))
+	{
+		return TArrayView<uint8>();
+	}
+	return TArrayView<uint8>(Ptr, Size);
+}
+
+////////////////////////////////////////////////////////////////////////////////
+FString FTraceAuxiliary::GetSecureHost()
+{
+	return GTraceAuxiliary.GetSecureHost();
+}
+
+////////////////////////////////////////////////////////////////////////////////
+uint16 FTraceAuxiliary::GetSecurePort()
+{
+	return GTraceAuxiliary.GetSecurePort();
+}
+
+#endif // UE_TRACE_ALLOW_SECURE_TRACING
+
 ////////////////////////////////////////////////////////////////////////////////
 bool FTraceAuxiliary::Stop()
 {
@@ -2119,6 +2347,7 @@ void FTraceAuxiliary::Initialize(const TCHAR* CommandLine)
 
 	Desc.bUseImportantCache = (FParse::Param(CommandLine, TEXT("tracenocache")) == false);
 	Desc.OnConnectionFunc = &TraceAuxiliaryOnConnectionCallback;
+	Desc.OnCloseFunc = &TraceAuxiliaryOnCloseCallback;
 	Desc.OnUpdateFunc = &TraceAuxiliaryOnTraceUpdateCallback;
 
 	// Determine if the worker thread is allowed
@@ -2132,7 +2361,7 @@ void FTraceAuxiliary::Initialize(const TCHAR* CommandLine)
 		// If allowed start worker thread if trace is started using command line argument,
 		// otherwise start on demand. For the editor we always want to start a worker thread
 		// since there is a lot of time until frames are reliably pumped.
-		Desc.bUseWorkerThread = bWasConnectedFromCmdLine || UE_EDITOR;
+		Desc.bUseWorkerThread = bWasConnectedFromCmdLine || UE_EDITOR || IS_PROGRAM;
 		GWorkerThreadConfig = Desc.bUseWorkerThread ? EWorkerThreadConfig::OnInit : EWorkerThreadConfig::OnDemand;
 	}
 
diff --git a/Engine/Source/Runtime/Core/Public/ProfilingDebugging/TraceAuxiliary.h b/Engine/Source/Runtime/Core/Public/ProfilingDebugging/TraceAuxiliary.h
--- a/Engine/Source/Runtime/Core/Public/ProfilingDebugging/TraceAuxiliary.h
+++ b/Engine/Source/Runtime/Core/Public/ProfilingDebugging/TraceAuxiliary.h
@@ -2,6 +2,7 @@
 
 #pragma once
 
+#include "Containers/AnsiString.h"
 #include "Containers/StringFwd.h"
 #include "CoreGlobals.h"
 #include "Delegates/Delegate.h"
@@ -89,6 +90,12 @@ public:
 		 * Relay connection. Pass user defined I/O functions to write raw trace data.
 		 */
 		Relay,
+#if UE_TRACE_ALLOW_SECURE_TRACING
+		/**
+		 * Connect to a receiver using a secure socket connection.
+		 */
+		SecureNetwork,
+#endif
 		/**
 		 * Don't connect, just start tracing to memory.
 		 */
@@ -100,13 +107,19 @@ public:
 	 */
 	DECLARE_TS_MULTICAST_DELEGATE(FOnConnection);
 
+	/**
+	 * Callback when a connection is closed.
+	 */
+	DECLARE_TS_MULTICAST_DELEGATE(FOnClosed);
+
 	/** Callback whenever a trace is started */
 	DECLARE_TS_MULTICAST_DELEGATE_TwoParams(FOnTraceStarted, FTraceAuxiliary::EConnectionType TraceType, const FString& TraceDestination);
 
 	/**
 	 * Callback whenever a trace recording is stopped.
-	 * TraceType tells what kind of trace it is.
-	 * TraceDestination will be either the filename and path for a file trace or the network connection for a network trace
+	 * @param TraceType Kind of trace
+	 * @param TraceDestination Either filename and path for a file trace or the network connection for a network trace
+	 * @note Callback will not trigger if trace is stopped by TraceLog itself, for example network interuptions or disk write errors.
 	 */
 	DECLARE_TS_MULTICAST_DELEGATE_TwoParams(FOnTraceStopped, FTraceAuxiliary::EConnectionType TraceType, const FString& TraceDestination);
 
@@ -150,6 +163,67 @@ public:
 	 */
 	static CORE_API bool Relay(UPTRINT Handle, UE::Trace::IoWriteFunc WriteFunc, UE::Trace::IoCloseFunc CloseFunc, const TCHAR* Channels = TEXT("default"), const FOptions* Options = nullptr);
 
+
+#if UE_TRACE_ALLOW_SECURE_TRACING
+
+	enum class ESecureMode : uint8
+	{
+		/**
+		 * Creates a self signed certficate, accessible via /ref GetSecureCertificate. Connecting peer is 
+		 * expected to send the authorization token before data will be received.
+		 */
+		SelfSignedCertAndAuthToken,
+		/**
+		 * Use a externally generated certificate and private key. Connecting peer is expected to have 
+		 * loaded the same certificate. Connecting peer is expected to send the authorization token 
+		 * before data will be received.
+		 */
+		CertFileAndAuthToken,
+		/**
+		 * Use authorization token as Pre-Shared Key verification.
+		 */
+		AuthToken
+	};
+
+	struct FSecureOptions 
+	{
+		// Specifies identity of the peer expected to connect if mode is AuthToken
+		FAnsiString Identity;
+		// Authorization token
+		TArray<uint8> AuthorizationToken;
+		// Timeout for connecting steps. A zero value means no timeouts.
+		float TimeoutSeconds = 0;
+	};
+
+	/**
+	 * Start tracing to network using a secure socket. This method is different from other tracing methods as it requires
+	 * the receiving peer to connect to this machine. Note that some options are controlled at the application level.
+	 * @param SecureOptions Options for how the secure connection is made
+	 * @param Options Optional additional tracing options.
+	 * @param Channels Comma separated list of channels to enable. Default set of channels are enabled if argument is not specified. If the pointer is null no channels are enabled.
+	 */
+	static CORE_API bool StartSecure(FSecureOptions& SecureOptions, FOptions* Options = nullptr, const TCHAR* Channels = TEXT("default"));
+
+	/**
+	 * Retrieve the certificate if secure tracing was initialized with using self-signed certificates.
+	 * @return The certificate in DER format if using self-signed certificates. Otherwise an empty view. The view is valid for the remainder of the
+	 * process lifetime.
+	 */
+	static CORE_API TArrayView<uint8> GetSecureCertificate();
+
+	/**
+	 * Gets the configured hostname connecting peers needs to use.
+	 * @return Hostname or IP
+	 */
+	static CORE_API FString GetSecureHost();
+
+	/**
+	 * Gets the configured port connecting peers should use.
+	 * @return Port number secure connection will listen to.
+	 */
+	static CORE_API uint16 GetSecurePort();
+#endif
+
 	/**
 	 * Stop tracing.
 	 * @return True if the trace was stopped, false if there was no data connection.
@@ -303,6 +377,11 @@ public:
 	 */
 	static CORE_API FOnConnection OnConnection;
 
+	/**
+	 * Delegate that triggers when a connection is closed.
+	 */
+	static CORE_API FOnClosed OnClosed;
+
 	/**
 	 * Delegate that triggers when a trace session is started.
 	 * The type of recording and the destination (filepath or network) is passed to the delegate.
