[Containers]

Adds TInstrumentationSafeCompactSetAllocator and renames  TInstrumentationSafeSetAllocator to be TInstrumentationSafeSparseSetAllocator allowing for TInstrumentationSafeSparseSetAllocator to map to whichever is being used by default via UE_USE_COMPACT_SET_AS_DEFAULT

#rnx

#Synced-CL 50001916

diff --git a/Engine/Source/Runtime/Core/Internal/Instrumentation/Containers.h b/Engine/Source/Runtime/Core/Internal/Instrumentation/Containers.h
--- a/Engine/Source/Runtime/Core/Internal/Instrumentation/Containers.h
+++ b/Engine/Source/Runtime/Core/Internal/Instrumentation/Containers.h
@@ -217,7 +217,7 @@ template<
 	uint32   BaseNumberOfHashBuckets = DEFAULT_BASE_NUMBER_OF_HASH_BUCKETS,
 	uint32   MinNumberOfHashedElements = DEFAULT_MIN_NUMBER_OF_HASHED_ELEMENTS
 >
-class TInstrumentationSafeSetAllocator
+class TInstrumentationSafeSparseSetAllocator
 {
 public:
 
@@ -240,6 +240,30 @@ public:
 	
 };
 
+template<typename BaseAllocator = TInstrumentationSafeBaseAllocator>
+struct TInstrumentationSafeCompactSetAllocator
+{
+	using InElementAllocator = TInstrumentationSafeAllocator<32, BaseAllocator>;
+
+	/** CompactSets use raw uint8, so we need extra alignment information for inline allocators */
+	template <typename ElementType>
+	struct AllocatorAlignment
+	{
+		static constexpr size_t Value = alignof(typename InElementAllocator::template ForElementType<uint8>);
+	};
+
+	template<int32 ElementSize>
+	using ElementAllocator = InElementAllocator;
+};
+
+#if UE_USE_COMPACT_SET_AS_DEFAULT
+template<typename BaseAllocator>
+using TInstrumentationSafeSetAllocator = TInstrumentationSafeCompactSetAllocator<BaseAllocator>;
+#else
+template<typename BaseAllocator>
+using TInstrumentationSafeSetAllocator = TInstrumentationSafeSparseSetAllocator<BaseAllocator>;
+#endif
+
 template <typename KeyType, typename ValueType, typename BaseAllocator = TInstrumentationSafeBaseAllocator>
 using TSafeMap = TMap<KeyType, ValueType, TInstrumentationSafeSetAllocator<BaseAllocator>>;
 
