StringBuilder: Remove deprecated functionality

#rnx
#Synced-CL 49858055

diff --git a/Engine/Source/Runtime/Core/Public/Misc/StringBuilder.h b/Engine/Source/Runtime/Core/Public/Misc/StringBuilder.h
--- a/Engine/Source/Runtime/Core/Public/Misc/StringBuilder.h
+++ b/Engine/Source/Runtime/Core/Public/Misc/StringBuilder.h
@@ -138,14 +138,6 @@ public:
 		return Base;
 	}
 
-	/** Prefer operator*() for a pointer to a null-terminated string. */
-	UE_DEPRECATED(5.6, "Call the non-const operator*().")
-	inline const CharType* ToString() const UE_LIFETIMEBOUND
-	{
-		const_cast<TStringBuilderBase*>(this)->EnsureNulTerminated();
-		return Base;
-	}
-
 	/** Returns a pointer to a null-terminated string that is valid until the builder is mutated. */
 	inline const CharType* operator*() UE_LIFETIMEBOUND
 	{
@@ -153,14 +145,6 @@ public:
 		return Base;
 	}
 
-	/** Returns a pointer to a null-terminated string that is valid until the builder is mutated. */
-	UE_DEPRECATED(5.6, "Call the non-const operator*().")
-	inline const CharType* operator*() const UE_LIFETIMEBOUND
-	{
-		const_cast<TStringBuilderBase*>(this)->EnsureNulTerminated();
-		return Base;
-	}
-
 	/** Returns a view of the string that is valid until the builder is mutated. */
 	inline ViewType ToView() const UE_LIFETIMEBOUND
 	{
@@ -630,25 +614,6 @@ inline TStringBuilderBase<CharType>& operator<<(TStringBuilderBase<CharType>& Bu
 	return Builder.Appendf(TFormatSpecifier<std::remove_cv_t<T>>::template GetFormatSpecifier<CharType>(), Value);
 }
 
-template <typename CharType, int32 BufferSize>
-class UE_DEPRECATED(5.3, "Use WriteToString<N>(...) or TStringBuilder<N>(InPlace, ...).") TWriteToString : public TStringBuilderWithBuffer<CharType, BufferSize>
-{
-public:
-	template <typename... ArgTypes>
-	explicit TWriteToString(ArgTypes&&... Args)
-	{
-		(*this << ... << (ArgTypes&&)Args);
-	}
-};
-
-PRAGMA_DISABLE_DEPRECATION_WARNINGS
-template <typename CharType, int32 BufferSize>
-struct TIsContiguousContainer<TWriteToString<CharType, BufferSize>>
-{
-	static constexpr inline bool Value = true;
-};
-PRAGMA_ENABLE_DEPRECATION_WARNINGS
-
 /**
  * A function to create and append to a temporary string builder.
  *
