Change the PRAGMA_PUSH_PLATFORM_DEFAULT_PACKING and PRAGMA_POP_PLATFORM_DEFAULT_PACKING macros for Clang to actually push and pop default packing.

#13375

#Synced-CL 50040876

diff --git a/Engine/Source/Runtime/Core/Public/Clang/ClangPlatformCompilerPreSetup.h b/Engine/Source/Runtime/Core/Public/Clang/ClangPlatformCompilerPreSetup.h
--- a/Engine/Source/Runtime/Core/Public/Clang/ClangPlatformCompilerPreSetup.h
+++ b/Engine/Source/Runtime/Core/Public/Clang/ClangPlatformCompilerPreSetup.h
@@ -305,11 +305,14 @@
 #endif // PRAGMA_POP
 
 #ifndef PRAGMA_PUSH_PLATFORM_DEFAULT_PACKING
-	#define PRAGMA_PUSH_PLATFORM_DEFAULT_PACKING
+	#define PRAGMA_PUSH_PLATFORM_DEFAULT_PACKING \
+		_Pragma("pack(push)") \
+		_Pragma("pack()")
 #endif // PRAGMA_PUSH_PLATFORM_DEFAULT_PACKING
 	
 #ifndef PRAGMA_POP_PLATFORM_DEFAULT_PACKING
-	#define PRAGMA_POP_PLATFORM_DEFAULT_PACKING
+	#define PRAGMA_POP_PLATFORM_DEFAULT_PACKING \
+		_Pragma("pack(pop)")
 #endif // PRAGMA_POP_PLATFORM_DEFAULT_PACKING
 
 #ifndef EMIT_CUSTOM_WARNING_AT_LINE
