Moving CsvProfiler code out of MallocBinnedCommon and into a file that's not included by the entire engine.
Fixing nonunity fallout from removing the CsvProfiler include.

#Synced-CL 49960643

diff --git a/Engine/Source/Runtime/ApplicationCore/Private/Windows/WindowsApplication.cpp b/Engine/Source/Runtime/ApplicationCore/Private/Windows/WindowsApplication.cpp
--- a/Engine/Source/Runtime/ApplicationCore/Private/Windows/WindowsApplication.cpp
+++ b/Engine/Source/Runtime/ApplicationCore/Private/Windows/WindowsApplication.cpp
@@ -2,6 +2,7 @@
 
 #include "Windows/WindowsApplication.h"
 #include "Null/NullApplication.h"
+#include "Containers/Queue.h"
 #include "Containers/StringConv.h"
 #include "CoreGlobals.h"
 #include "Internationalization/Text.h"
diff --git a/Engine/Source/Runtime/Core/Private/HAL/MallocBinned2.cpp b/Engine/Source/Runtime/Core/Private/HAL/MallocBinned2.cpp
--- a/Engine/Source/Runtime/Core/Private/HAL/MallocBinned2.cpp
+++ b/Engine/Source/Runtime/Core/Private/HAL/MallocBinned2.cpp
@@ -953,7 +953,7 @@ void FMallocBinned2::UpdateStats()
 	CSV_CUSTOM_STAT(FMemory, AllocatorCachedSlackMB, (int32)(CachedOSPageAllocator.GetCachedFreeTotal() / (1024 * 1024)), ECsvCustomStatOp::Set);
 	CSV_CUSTOM_STAT(FMemory, AllocatorImmediatelyFreeableCachedSlackMB, (int32)(CachedOSPageAllocator.GetCachedImmediatelyFreeable() / (1024 * 1024)), ECsvCustomStatOp::Set);
 
-	UpdateStatsCommon(*this);
+	FMallocBinnedCommonUtils::UpdateStatsCommon(*this);
 	CachedOSPageAllocator.UpdateStats();
 	FScopedVirtualMallocTimer::UpdateStats();
 }
diff --git a/Engine/Source/Runtime/Core/Private/HAL/MallocBinned3.cpp b/Engine/Source/Runtime/Core/Private/HAL/MallocBinned3.cpp
--- a/Engine/Source/Runtime/Core/Private/HAL/MallocBinned3.cpp
+++ b/Engine/Source/Runtime/Core/Private/HAL/MallocBinned3.cpp
@@ -13,7 +13,6 @@ PRAGMA_DISABLE_UNSAFE_TYPECAST_WARNINGS
 #include "HAL/PlatformMisc.h"
 #include "HAL/MallocBinnedCommonUtils.h"
 
-
 #if UE_MB3_USE_CACHED_PAGE_ALLOCATOR_FOR_LARGE_ALLOCS
 #	include "HAL/Allocators/CachedOSPageAllocator.h"
 #	define UE_MB3_MAX_CACHED_OS_FREES (64)
@@ -1116,6 +1115,11 @@ void FMallocBinned3::DumpAllocatorStats(class FOutputDevice& Ar)
 #endif
 }
 
+void FMallocBinned3::UpdateStats()
+{
+	FMallocBinnedCommonUtils::UpdateStatsCommon(*this);
+}
+
 #endif //~PLATFORM_HAS_FPlatformVirtualMemoryBlock
 
 PRAGMA_RESTORE_UNSAFE_TYPECAST_WARNINGS
diff --git a/Engine/Source/Runtime/Core/Public/HAL/MallocBinned3.h b/Engine/Source/Runtime/Core/Public/HAL/MallocBinned3.h
--- a/Engine/Source/Runtime/Core/Public/HAL/MallocBinned3.h
+++ b/Engine/Source/Runtime/Core/Public/HAL/MallocBinned3.h
@@ -348,10 +348,7 @@ public:
 	virtual const TCHAR* GetDescriptiveName() override;
 	/** Dumps current allocator stats to the log. */
 	virtual void DumpAllocatorStats(class FOutputDevice& Ar) override;
-	virtual void UpdateStats() override
-	{
-		UpdateStatsCommon(*this);
-	}
+	virtual void UpdateStats() override;
 	// End FMalloc interface.
 
 	void FreeExternal(void* Ptr, uint64 PoolIndex);
diff --git a/Engine/Source/Runtime/Core/Public/HAL/MallocBinnedCommon.h b/Engine/Source/Runtime/Core/Public/HAL/MallocBinnedCommon.h
--- a/Engine/Source/Runtime/Core/Public/HAL/MallocBinnedCommon.h
+++ b/Engine/Source/Runtime/Core/Public/HAL/MallocBinnedCommon.h
@@ -12,8 +12,6 @@
 #include "Misc/ScopeLock.h"
 #include "Templates/AlignmentTemplates.h"
 #include "Templates/MemoryOps.h"
-#include "ProfilingDebugging/CsvProfiler.h"
-
 
 // A project can define it's own UE_MBC_MAX_LISTED_SMALL_POOL_SIZE and UE_MBC_NUM_LISTED_SMALL_POOLS to reduce runtime memory usage
 // MallocBinnedCommon.cpp has a list of predefined bins that go up to 28672
@@ -108,8 +106,6 @@
 
 #define UE_MBC_MAX_SUPPORTED_PLATFORM_PAGE_SIZE			16 * 1024	// MB3 and it's bins were designed around 4\16KB native page size. Having larger page size will increase memory waste and will be inefficient for current allocator design
 
-CSV_DECLARE_CATEGORY_EXTERN(MallocBinned);
-
 extern CORE_API float GMallocBinnedFlushThreadCacheMaxWaitTime;
 extern CORE_API int32 GMallocBinnedFlushRegisteredThreadCachesOnOneThread;
 
@@ -953,59 +949,6 @@ protected:
 	}
 #endif
 
-	void UpdateStatsCommon(const AllocType& Alloc)
-	{
-#if UE_MBC_ALLOCATOR_STATS && CSV_PROFILER_STATS 
-		if (!GMallocBinnedEnableCSVStats && !FCsvProfiler::Get()->IsCategoryEnabled(CSV_CATEGORY_INDEX(MallocBinned)))
-		{
-			return;
-		}
-
-		FCsvProfiler::Get()->EnableCategoryByIndex(CSV_CATEGORY_INDEX(MallocBinned), true);
-
-		static bool bFirstTime = true;
-		static FName FragmentationsStats[NumSmallPools];
-		static FName WasteStats[NumSmallPools];
-		static FName TotalMemStats[NumSmallPools];
-
-		if (bFirstTime)
-		{
-			for (int32 i = 0; i < NumSmallPools; i++)
-			{
-				const int BinSize = Alloc.PoolIndexToBinSize(i);
-
-				TCHAR Name[64];
-				FCString::Sprintf(Name, TEXT("FragmentationBin%d"), BinSize);
-				FragmentationsStats[i] = FName(Name);
-
-				FCString::Sprintf(Name, TEXT("WasteBin%d"), BinSize);
-				WasteStats[i] = FName(Name);
-
-				FCString::Sprintf(Name, TEXT("TotalMemBin%d"), BinSize);
-				TotalMemStats[i] = FName(Name);
-			}
-
-			bFirstTime = false;
-		}
-
-		for (int32 i = 0; i < NumSmallPools; i++)
-		{
-			const float Fragmentation = 1.0f - (float)Alloc.SmallPoolTables[i].TotalUsedBins / (float)Alloc.SmallPoolTables[i].TotalAllocatedBins;
-			FCsvProfiler::RecordCustomStat(FragmentationsStats[i], CSV_CATEGORY_INDEX(MallocBinned), int(Fragmentation * 100.0f), ECsvCustomStatOp::Set);
-
-			const float TotalMem = (float)Alloc.SmallPoolTables[i].TotalAllocatedMem / 1024.0f / 1024.0f;
-			FCsvProfiler::RecordCustomStat(TotalMemStats[i], CSV_CATEGORY_INDEX(MallocBinned), TotalMem, ECsvCustomStatOp::Set);
-
-			FCsvProfiler::RecordCustomStat(WasteStats[i], CSV_CATEGORY_INDEX(MallocBinned), TotalMem * Fragmentation, ECsvCustomStatOp::Set);
-		}
-
-		CSV_CUSTOM_STAT(MallocBinned, RequestedSmallPoolMemoryMB,		(float)GetTotalAllocatedSmallPoolMemory() / (1024.0f * 1024.0f), ECsvCustomStatOp::Set);
-		CSV_CUSTOM_STAT(MallocBinned, TotalAllocatedSmallPoolMemoryMB,   (float)AllocatedOSSmallPoolMemory.load(std::memory_order_relaxed) / (1024.0f * 1024.0f), ECsvCustomStatOp::Set);
-		CSV_CUSTOM_STAT(MallocBinned, RequestedLargeAllocsMemoryMB,		(float)AllocatedLargePoolMemory.load(std::memory_order_relaxed) / (1024.0f * 1024.0f), ECsvCustomStatOp::Set);
-		CSV_CUSTOM_STAT(MallocBinned, TotalAllocatedLargeAllocsMemoryMB, (float)AllocatedLargePoolMemoryWAlignment.load(std::memory_order_relaxed) / (1024.0f * 1024.0f), ECsvCustomStatOp::Set);
-#endif
-	}
-
 	void AllocateHashBuckets()
 	{
 		const uint64 MaxHashBuckets = PtrToPoolMapping.GetMaxHashBuckets();
@@ -1040,3 +983,7 @@ private:
 #if UE_ENABLE_INCLUDE_ORDER_DEPRECATED_IN_5_7
 #	include "HAL/CriticalSection.h"
 #endif
+
+#if UE_ENABLE_INCLUDE_ORDER_DEPRECATED_IN_5_8
+#	#include "ProfilingDebugging/CsvProfiler.h"
+#endif
diff --git a/Engine/Source/Runtime/Core/Public/HAL/MallocBinnedCommonUtils.h b/Engine/Source/Runtime/Core/Public/HAL/MallocBinnedCommonUtils.h
--- a/Engine/Source/Runtime/Core/Public/HAL/MallocBinnedCommonUtils.h
+++ b/Engine/Source/Runtime/Core/Public/HAL/MallocBinnedCommonUtils.h
@@ -8,7 +8,9 @@
 #include "Misc/App.h"
 #include "Stats/Stats.h"
 #include "FramePro/FrameProProfiler.h"
+#include "ProfilingDebugging/CsvProfiler.h"
 
+CSV_DECLARE_CATEGORY_EXTERN(MallocBinned);
 
 #if FRAMEPRO_ENABLED
 //	Pushes a profiler scope if it's safe to do so without any new allocations.
@@ -185,4 +187,60 @@ public:
 			FTaskGraphInterface::BroadcastSlow_OnlyUseForSpecialPurposes(FPlatformProcess::SupportsMultithreading() && FApp::ShouldUseThreadingForPerformance(), false, Broadcast);
 		}
 	}
+
+	template <class AllocType>
+	static void UpdateStatsCommon(const AllocType& Alloc)
+	{
+#if UE_MBC_ALLOCATOR_STATS && CSV_PROFILER_STATS 
+		if (!GMallocBinnedEnableCSVStats && !FCsvProfiler::Get()->IsCategoryEnabled(CSV_CATEGORY_INDEX(MallocBinned)))
+		{
+			return;
+		}
+
+		constexpr int32 NumSmallPools = AllocType::NUM_SMALL_POOLS;
+
+		FCsvProfiler::Get()->EnableCategoryByIndex(CSV_CATEGORY_INDEX(MallocBinned), true);
+
+		static bool bFirstTime = true;
+		static FName FragmentationsStats[NumSmallPools];
+		static FName WasteStats[NumSmallPools];
+		static FName TotalMemStats[NumSmallPools];
+
+		if (bFirstTime)
+		{
+			for (int32 i = 0; i < NumSmallPools; i++)
+			{
+				const int BinSize = Alloc.PoolIndexToBinSize(i);
+
+				TCHAR Name[64];
+				FCString::Sprintf(Name, TEXT("FragmentationBin%d"), BinSize);
+				FragmentationsStats[i] = FName(Name);
+
+				FCString::Sprintf(Name, TEXT("WasteBin%d"), BinSize);
+				WasteStats[i] = FName(Name);
+
+				FCString::Sprintf(Name, TEXT("TotalMemBin%d"), BinSize);
+				TotalMemStats[i] = FName(Name);
+			}
+
+			bFirstTime = false;
+		}
+
+		for (int32 i = 0; i < NumSmallPools; i++)
+		{
+			const float Fragmentation = 1.0f - (float)Alloc.SmallPoolTables[i].TotalUsedBins / (float)Alloc.SmallPoolTables[i].TotalAllocatedBins;
+			FCsvProfiler::RecordCustomStat(FragmentationsStats[i], CSV_CATEGORY_INDEX(MallocBinned), int(Fragmentation * 100.0f), ECsvCustomStatOp::Set);
+
+			const float TotalMem = (float)Alloc.SmallPoolTables[i].TotalAllocatedMem / 1024.0f / 1024.0f;
+			FCsvProfiler::RecordCustomStat(TotalMemStats[i], CSV_CATEGORY_INDEX(MallocBinned), TotalMem, ECsvCustomStatOp::Set);
+
+			FCsvProfiler::RecordCustomStat(WasteStats[i], CSV_CATEGORY_INDEX(MallocBinned), TotalMem * Fragmentation, ECsvCustomStatOp::Set);
+		}
+
+		CSV_CUSTOM_STAT(MallocBinned, RequestedSmallPoolMemoryMB, (float)Alloc.GetTotalAllocatedSmallPoolMemory() / (1024.0f * 1024.0f), ECsvCustomStatOp::Set);
+		CSV_CUSTOM_STAT(MallocBinned, TotalAllocatedSmallPoolMemoryMB, (float)Alloc.AllocatedOSSmallPoolMemory.load(std::memory_order_relaxed) / (1024.0f * 1024.0f), ECsvCustomStatOp::Set);
+		CSV_CUSTOM_STAT(MallocBinned, RequestedLargeAllocsMemoryMB, (float)Alloc.AllocatedLargePoolMemory.load(std::memory_order_relaxed) / (1024.0f * 1024.0f), ECsvCustomStatOp::Set);
+		CSV_CUSTOM_STAT(MallocBinned, TotalAllocatedLargeAllocsMemoryMB, (float)Alloc.AllocatedLargePoolMemoryWAlignment.load(std::memory_order_relaxed) / (1024.0f * 1024.0f), ECsvCustomStatOp::Set);
+#endif
+	}
 };
