LLM: Bugfix for tracking sizes of ELLMTagSets other than ELLMTagSet::None. Fix double (or more) counting of allocations caused by multiplying allocation size by the number of parents of the Allocation's ELLMTagSet::None tag.
#rnx

#Synced-CL 49917028

diff --git a/Engine/Source/Runtime/Core/Private/HAL/LLM/LLM.cpp b/Engine/Source/Runtime/Core/Private/HAL/LLM/LLM.cpp
--- a/Engine/Source/Runtime/Core/Private/HAL/LLM/LLM.cpp
+++ b/Engine/Source/Runtime/Core/Private/HAL/LLM/LLM.cpp
@@ -5711,10 +5711,24 @@ void FLLMTracker::FetchAndClearTagSizes_UpdateGroupAllocs(FAllocationGroup& Allo
 		PruningLevel);
 	if (Size != 0)
 	{
-		for (const FTagData* TagData : AllocationGroup.GetActiveTags())
-		{
-			FTrackerTagSizeData& TrackerInfo = TagSizes.FindOrAdd(TagData);
-			TrackerInfo.Size += Size;
+		// AllocationGroup sizes are accumulated upwards into parent AllocationGroups, where a parent
+		// AllocationGroup is the AllocationGroup with all tags the same except for Systems tag which has
+		// been replaced with its parent. When copying sizes from AllocationGroup to Tag, for all the
+		// other non-hierarchical tags, we need to only copy the sizes from the topmost parent, to avoid
+		// counting multiple times the size from a child AllocationGroup that has been added to all of
+		// its ancestor group's sizes.
+		const FActiveTags& ActiveTags = AllocationGroup.GetActiveTags();
+		int32 NumActiveTags = ActiveTags.Num();
+		const FTagData* SystemsTag = ActiveTags.GetSystemsTagData();
+		LLMCheck(SystemsTag == ActiveTags[0]);
+		TagSizes.FindOrAdd(SystemsTag).Size += Size;
+		if (SystemsTag->GetParent() == nullptr)
+		{
+			for (int32 TagIndex = 1; TagIndex < NumActiveTags; ++TagIndex)
+			{
+				const FTagData* NonHierarchicalTag = ActiveTags[TagIndex];
+				TagSizes.FindOrAdd(NonHierarchicalTag).Size += Size;
+			}
 		}
 	}
 
