Get AutoRTFM to initially work with ICU.

#Synced-CL 50046643

diff --git a/Engine/Source/Programs/AutoRTFMTests/AutoRTFMTests.Target.cs b/Engine/Source/Programs/AutoRTFMTests/AutoRTFMTests.Target.cs
--- a/Engine/Source/Programs/AutoRTFMTests/AutoRTFMTests.Target.cs
+++ b/Engine/Source/Programs/AutoRTFMTests/AutoRTFMTests.Target.cs
@@ -33,10 +33,10 @@ public class AutoRTFMTestsTarget : TargetRules
 		// Make a console application under Windows, so entry point is main() everywhere
 		bIsBuildingConsoleApplication = true;
 
-		// Disable unity builds by default for AutoRTFMTest
+		// Disable unity builds (stresses all the corners of the AutoRTFM runtime lookup code).
 		bUseUnityBuild = false;
 
-		// Disable ICU to match the server config
+		// Server has ICU off so use that as the default.
 		bCompileICU = false;
 
 		// Set the RTFM clang compiler
diff --git a/Engine/Source/Programs/AutoRTFMTests/AutoRTFMTestsWithICU.Target.cs b/Engine/Source/Programs/AutoRTFMTests/AutoRTFMTestsWithICU.Target.cs
new file mode 100644
--- /dev/null
+++ b/Engine/Source/Programs/AutoRTFMTests/AutoRTFMTestsWithICU.Target.cs
@@ -0,0 +1,16 @@
+// Copyright Epic Games, Inc. All Rights Reserved.
+
+using EpicGames.Core;
+using UnrealBuildTool;
+using System.Collections.Generic;
+using UnrealBuildBase;
+
+[SupportedPlatforms(UnrealPlatformClass.All)]
+public class AutoRTFMTestsWithICUTarget : AutoRTFMTestsTarget
+{
+	public AutoRTFMTestsWithICUTarget(TargetInfo Target) : base(Target)
+	{
+		// Editor uses ICU so we need to test with it too.
+		bCompileICU = true;
+	}
+}
diff --git a/Engine/Source/Programs/AutoRTFMTests/Private/UECoreTests.cpp b/Engine/Source/Programs/AutoRTFMTests/Private/UECoreTests.cpp
--- a/Engine/Source/Programs/AutoRTFMTests/Private/UECoreTests.cpp
+++ b/Engine/Source/Programs/AutoRTFMTests/Private/UECoreTests.cpp
@@ -18,6 +18,7 @@
 #include "Internationalization/TextCache.h"
 #include "Internationalization/TextFormatter.h"
 #include "Internationalization/TextHistory.h"
+#include "Internationalization/TextTransformer.h"
 #include "Logging/LogMacros.h"
 #include "Logging/StructuredLog.h"
 #include "Logging/StructuredLogFormat.h"
@@ -2465,4 +2466,25 @@ TEST_CASE("UECore.FConcurrentLinearBulkObjectAllocator")
 	}
 }
 
+TEST_CASE("UECore.TextTransformer")
+{
+	SECTION("ToLower")
+	{
+		FString String(TEXT("WOWWEE"));
+
+		AutoRTFM::Testing::Commit([&]
+		{
+			REQUIRE(FTextTransformer::ToLower(String).Equals(TEXT("wowwee")));
+		});
+	}
 
+	SECTION("ToUpper")
+	{
+		FString String(TEXT("wowwee"));
+
+		AutoRTFM::Testing::Commit([&]
+		{
+			REQUIRE(FTextTransformer::ToUpper(String).Equals(TEXT("WOWWEE")));
+		});
+	}
+}
diff --git a/Engine/Source/Runtime/AutoRTFM/Public/AutoRTFM.h b/Engine/Source/Runtime/AutoRTFM/Public/AutoRTFM.h
--- a/Engine/Source/Runtime/AutoRTFM/Public/AutoRTFM.h
+++ b/Engine/Source/Runtime/AutoRTFM/Public/AutoRTFM.h
@@ -187,11 +187,11 @@ UE_AUTORTFM_API void SetMode(EMode Mode);
 // Calls to BeginDisable() are internally counted and may be nested within
 // other BeginDisable() / EndDisable() calls.
 // This call is a no-op when called outside of a transaction.
-UE_AUTORTFM_API AUTORTFM_OPEN void BeginDisable();
+AUTORTFM_OPEN UE_AUTORTFM_API void BeginDisable();
 
 // Ends temporarily disabling the AutoRTFM sanitizer.
 // See BeginDisable().
-UE_AUTORTFM_API AUTORTFM_OPEN void EndDisable();
+AUTORTFM_OPEN UE_AUTORTFM_API void EndDisable();
 
 // Returns true if the AutoRTFM sanitizer records the callstack of closed writes.
 UE_AUTORTFM_API bool RecordClosedWriteCallstacks();
diff --git a/Engine/Source/Runtime/Core/Private/Internationalization/ICUCulture.cpp b/Engine/Source/Runtime/Core/Private/Internationalization/ICUCulture.cpp
--- a/Engine/Source/Runtime/Core/Private/Internationalization/ICUCulture.cpp
+++ b/Engine/Source/Runtime/Core/Private/Internationalization/ICUCulture.cpp
@@ -206,6 +206,11 @@ TArray<ETextPluralForm> ICUPluralRulesToUEValidPluralForms(const icu::PluralRule
 FICUCultureImplementation::FICUCultureImplementation(const FString& LocaleName)
 	: ICULocale( TCHAR_TO_ANSI( *LocaleName ) )
 {
+	// If anyone ever tries to remove this in future you need to remove the `AUTORTFM_OPEN` from
+	// `GetDecimalNumberFormattingRules`, `GetPercentFormattingRules`, and `GetCurrencyFormattingRules`,
+	// and actually make all that code transactionally safe (its gonna be tricky!).
+	AutoRTFM::UnreachableIfClosed("Cannot construct `FICUCultureImplementation` in the closed!");
+
 	if (ICULocale.isBogus())
 	{
 		ICULocale = GetInvariantLocale();
diff --git a/Engine/Source/Runtime/Core/Private/Internationalization/ICUCulture.h b/Engine/Source/Runtime/Core/Private/Internationalization/ICUCulture.h
--- a/Engine/Source/Runtime/Core/Private/Internationalization/ICUCulture.h
+++ b/Engine/Source/Runtime/Core/Private/Internationalization/ICUCulture.h
@@ -189,9 +189,15 @@ public:
 	virtual FString GetScript() const override;
 	virtual FString GetVariant() const override;
 	virtual bool IsRightToLeft() const override;
-	virtual const FDecimalNumberFormattingRules& GetDecimalNumberFormattingRules() override;
-	virtual const FDecimalNumberFormattingRules& GetPercentFormattingRules() override;
-	virtual const FDecimalNumberFormattingRules& GetCurrencyFormattingRules(const FString& InCurrencyCode) override;
+
+	// We run this in the open which is safe so long as we cannot create `FICUCultureImplementation` in the closed.
+	AUTORTFM_OPEN virtual const FDecimalNumberFormattingRules& GetDecimalNumberFormattingRules() override;
+
+	// We run this in the open which is safe so long as we cannot create `FICUCultureImplementation` in the closed.
+	AUTORTFM_OPEN virtual const FDecimalNumberFormattingRules& GetPercentFormattingRules() override;
+
+	// We run this in the open which is safe so long as we cannot create `FICUCultureImplementation` in the closed.
+	AUTORTFM_OPEN virtual const FDecimalNumberFormattingRules& GetCurrencyFormattingRules(const FString& InCurrencyCode) override;
 	virtual ETextPluralForm GetPluralForm(int32 Val, const ETextPluralType PluralType) const override;
 	virtual ETextPluralForm GetPluralForm(double Val, const ETextPluralType PluralType) const override;
 	virtual const TArray<ETextPluralForm>& GetValidPluralForms(const ETextPluralType PluralType) const override;
diff --git a/Engine/Source/Runtime/Core/Private/Internationalization/ICUText.cpp b/Engine/Source/Runtime/Core/Private/Internationalization/ICUText.cpp
--- a/Engine/Source/Runtime/Core/Private/Internationalization/ICUText.cpp
+++ b/Engine/Source/Runtime/Core/Private/Internationalization/ICUText.cpp
@@ -85,7 +85,13 @@ FString FTextTransformer::ToLower(const FString& InStr)
 	FInternationalization& I18N = FInternationalization::Get();
 	checkf(I18N.IsInitialized() == true, TEXT("FInternationalization is not initialized. An FText formatting method was likely used in static object initialization - this is not supported."));
 	
-	return ICUUtilities::ConvertString(ICUUtilities::ConvertString(InStr).toLower());
+	// We run this code in the open because it calls into ICU which is not transactionally instrumented.
+	// We get away with this because we use the explicit `AutoRTFM::Open` which knows how to take open
+	// constructed `FString`'s and make them transactionally safe (see `AutoRTFMReturnFromOpenMode`).
+	return AutoRTFM::Open([&]
+	{
+		return ICUUtilities::ConvertString(ICUUtilities::ConvertString(InStr).toLower());
+	});
 }
 
 FString FTextTransformer::ToUpper(const FString& InStr)
@@ -93,7 +99,13 @@ FString FTextTransformer::ToUpper(const FString& InStr)
 	FInternationalization& I18N = FInternationalization::Get();
 	checkf(I18N.IsInitialized() == true, TEXT("FInternationalization is not initialized. An FText formatting method was likely used in static object initialization - this is not supported."));
 
-	return ICUUtilities::ConvertString(ICUUtilities::ConvertString(InStr).toUpper());
+	// We run this code in the open because it calls into ICU which is not transactionally instrumented.
+	// We get away with this because we use the explicit `AutoRTFM::Open` which knows how to take open
+	// constructed `FString`'s and make them transactionally safe (see `AutoRTFMReturnFromOpenMode`).
+	return AutoRTFM::Open([&]
+	{
+		return ICUUtilities::ConvertString(ICUUtilities::ConvertString(InStr).toUpper());
+	});
 }
 
 bool FTextChar::IsAlpha(const UTF32CHAR Codepoint)
