Add MakeArrayView64 equivalents to the 32-bit TArrayView MakeArrayView functions.
#rnx

#Synced-CL 49917028

diff --git a/Engine/Source/Runtime/Core/Public/Containers/ArrayView.h b/Engine/Source/Runtime/Core/Public/Containers/ArrayView.h
--- a/Engine/Source/Runtime/Core/Public/Containers/ArrayView.h
+++ b/Engine/Source/Runtime/Core/Public/Containers/ArrayView.h
@@ -866,6 +866,17 @@ template <
 {
 	return TArrayView<std::remove_pointer_t<decltype(GetData(DeclVal<OtherRangeType&>()))>>(Forward<OtherRangeType>(Other));
 }
+
+template <
+	typename OtherRangeType,
+	typename CVUnqualifiedOtherRangeType = std::remove_cv_t<std::remove_reference_t<OtherRangeType>>
+	UE_REQUIRES(TIsContiguousContainer<CVUnqualifiedOtherRangeType>::Value && TIsTArrayView_V<CVUnqualifiedOtherRangeType>)
+>
+[[nodiscard]] constexpr auto MakeArrayView64(OtherRangeType&& Other)
+{
+	return TArrayView64<std::remove_pointer_t<decltype(GetData(DeclVal<OtherRangeType&>()))>>(Forward<OtherRangeType>(Other));
+}
+
 template <
 	typename OtherRangeType,
 	typename CVUnqualifiedOtherRangeType = std::remove_cv_t<std::remove_reference_t<OtherRangeType>>
@@ -876,18 +887,40 @@ template <
 	return TArrayView<std::remove_pointer_t<decltype(GetData(DeclVal<OtherRangeType&>()))>>(Forward<OtherRangeType>(Other));
 }
 
+template <
+	typename OtherRangeType,
+	typename CVUnqualifiedOtherRangeType = std::remove_cv_t<std::remove_reference_t<OtherRangeType>>
+	UE_REQUIRES(TIsContiguousContainer<CVUnqualifiedOtherRangeType>::Value && !TIsTArrayView_V<CVUnqualifiedOtherRangeType>)
+>
+[[nodiscard]] constexpr auto MakeArrayView64(OtherRangeType&& Other UE_LIFETIMEBOUND)
+{
+	return TArrayView64<std::remove_pointer_t<decltype(GetData(DeclVal<OtherRangeType&>()))>>(Forward<OtherRangeType>(Other));
+}
+
 template<typename ElementType>
 [[nodiscard]] constexpr TArrayView<ElementType> MakeArrayView(ElementType* Pointer UE_LIFETIMEBOUND, int32 Size)
 {
 	return TArrayView<ElementType>(Pointer, Size);
 }
 
+template<typename ElementType>
+[[nodiscard]] constexpr TArrayView64<ElementType> MakeArrayView64(ElementType* Pointer UE_LIFETIMEBOUND, int64 Size)
+{
+	return TArrayView64<ElementType>(Pointer, Size);
+}
+
 template <typename ElementType>
 [[nodiscard]] constexpr TArrayView<const ElementType> MakeArrayView(std::initializer_list<ElementType> List UE_LIFETIMEBOUND)
 {
 	return TArrayView<const ElementType>(List.begin(), List.size());
 }
 
+template <typename ElementType>
+[[nodiscard]] constexpr TArrayView64<const ElementType> MakeArrayView64(std::initializer_list<ElementType> List UE_LIFETIMEBOUND)
+{
+	return TArrayView64<const ElementType>(List.begin(), List.size());
+}
+
 template <
 	typename OtherRangeType,
 	typename CVUnqualifiedOtherRangeType = std::remove_cv_t<std::remove_reference_t<OtherRangeType>>
@@ -897,6 +930,17 @@ template <
 {
 	return TArrayView<const std::remove_pointer_t<decltype(GetData(DeclVal<OtherRangeType&>()))>>(Forward<OtherRangeType>(Other));
 }
+
+template <
+	typename OtherRangeType,
+	typename CVUnqualifiedOtherRangeType = std::remove_cv_t<std::remove_reference_t<OtherRangeType>>
+	UE_REQUIRES(TIsContiguousContainer<CVUnqualifiedOtherRangeType>::Value && TIsTArrayView_V<CVUnqualifiedOtherRangeType>)
+>
+[[nodiscard]] constexpr auto MakeConstArrayView64(OtherRangeType&& Other)
+{
+	return TArrayView64<const std::remove_pointer_t<decltype(GetData(DeclVal<OtherRangeType&>()))>>(Forward<OtherRangeType>(Other));
+}
+
 template <
 	typename OtherRangeType,
 	typename CVUnqualifiedOtherRangeType = std::remove_cv_t<std::remove_reference_t<OtherRangeType>>
@@ -907,18 +951,40 @@ template <
 	return TArrayView<const std::remove_pointer_t<decltype(GetData(DeclVal<OtherRangeType&>()))>>(Forward<OtherRangeType>(Other));
 }
 
+template <
+	typename OtherRangeType,
+	typename CVUnqualifiedOtherRangeType = std::remove_cv_t<std::remove_reference_t<OtherRangeType>>
+	UE_REQUIRES(TIsContiguousContainer<CVUnqualifiedOtherRangeType>::Value && !TIsTArrayView_V<CVUnqualifiedOtherRangeType>)
+>
+[[nodiscard]] constexpr auto MakeConstArrayView64(OtherRangeType&& Other UE_LIFETIMEBOUND)
+{
+	return TArrayView64<const std::remove_pointer_t<decltype(GetData(DeclVal<OtherRangeType&>()))>>(Forward<OtherRangeType>(Other));
+}
+
 template<typename ElementType>
 [[nodiscard]] constexpr TArrayView<const ElementType> MakeConstArrayView(const ElementType* Pointer UE_LIFETIMEBOUND, int32 Size)
 {
 	return TArrayView<const ElementType>(Pointer, Size);
 }
 
+template<typename ElementType>
+[[nodiscard]] constexpr TArrayView64<const ElementType> MakeConstArrayView64(const ElementType* Pointer UE_LIFETIMEBOUND, int64 Size)
+{
+	return TArrayView64<const ElementType>(Pointer, Size);
+}
+
 template <typename ElementType>
 [[nodiscard]] constexpr TArrayView<const ElementType> MakeConstArrayView(std::initializer_list<ElementType> List UE_LIFETIMEBOUND)
 {
 	return TArrayView<const ElementType>(List.begin(), List.size());
 }
 
+template <typename ElementType>
+[[nodiscard]] constexpr TArrayView64<const ElementType> MakeConstArrayView64(std::initializer_list<ElementType> List UE_LIFETIMEBOUND)
+{
+	return TArrayView64<const ElementType>(List.begin(), List.size());
+}
+
 //////////////////////////////////////////////////////////////////////////
 
 template<typename InElementType, typename InAllocatorType>
