[Platforms]

- Adds utility method to GenericPlatformMisc to sanitize SystemWideMutex names to be used by the Core module.
- Adjusts the platform implementations for platform-specific FPlatformSystemWideMutex to sanitize the input name provided to remove any characters that would be invalid to have in a filename as well as the '.'. This sanitization is meant to ensure that any name provided by the user is usable by the platform creating the system wide mutex which often requires creating a handle with a valid path. Without sanitization input names containing .., ., ., \\, / etc... might be incorrectly used for creating the handle name and silently fail. Some sanitization was already done for some platforms but now we have a more robust common sanitization.

#rnx

#Synced-CL 50040740

#Synced-CL 50040740

diff --git a/Engine/Source/Runtime/Core/Private/GenericPlatform/GenericPlatformMisc.cpp b/Engine/Source/Runtime/Core/Private/GenericPlatform/GenericPlatformMisc.cpp
--- a/Engine/Source/Runtime/Core/Private/GenericPlatform/GenericPlatformMisc.cpp
+++ b/Engine/Source/Runtime/Core/Private/GenericPlatform/GenericPlatformMisc.cpp
@@ -2411,6 +2411,23 @@ int FGenericPlatformMisc::GetMobilePropagateAlphaSetting()
 	return PropagateAlphaCVar ? PropagateAlphaCVar->GetInt() : 0;
 }
 
+void FGenericPlatformMisc::SanitizeSystemWideMutexName(FString& InName)
+{
+	// Remove all invalid filesystem characters from the provided name but also 
+	// remove '.' since we don't want names with . or .. to be interpreted as 
+	// paths for platforms that might require handle names to be valid paths
+	const FString InvalidChars = FString::Printf(TEXT("%s."), *FPaths::GetInvalidFileSystemChars());
+
+	int32 Index = INDEX_NONE;
+	for (TCHAR& Char : InName)
+	{
+		if (InvalidChars.FindChar(Char, Index))
+		{
+			Char = TEXT('_');
+		}
+	}
+}
+
 void FGenericPlatformMisc::ShowConsoleWindow()
 {
 #if !UE_BUILD_SHIPPING
diff --git a/Engine/Source/Runtime/Core/Private/Mac/MacPlatformMutex.cpp b/Engine/Source/Runtime/Core/Private/Mac/MacPlatformMutex.cpp
--- a/Engine/Source/Runtime/Core/Private/Mac/MacPlatformMutex.cpp
+++ b/Engine/Source/Runtime/Core/Private/Mac/MacPlatformMutex.cpp
@@ -19,9 +19,9 @@ FMacSystemWideMutex::FMacSystemWideMutex(const FString& InName, FTimespan InTime
 	check(InTimeout >= FTimespan::Zero());
 	check(InTimeout.GetTotalSeconds() < (double)FLT_MAX);
 
-	const FString LockPath = FString(FPlatformProcess::ApplicationSettingsDir()) / InName;
-	FString NormalizedFilepath(LockPath);
-	NormalizedFilepath.ReplaceInline(TEXT("\\"), TEXT("/"));
+	FString SanitizedName(InName);
+	FGenericPlatformMisc::SanitizeSystemWideMutexName(SanitizedName);
+	const FString LockPath = FString(FPlatformProcess::ApplicationSettingsDir()) / SanitizedName;
 	FileHandle = -1;
 
 	double ExpireTimeSecs = FPlatformTime::Seconds() + InTimeout.GetTotalSeconds();
@@ -30,7 +30,7 @@ FMacSystemWideMutex::FMacSystemWideMutex(const FString& InName, FTimespan InTime
 		if (FileHandle == -1)
 		{
 			// Try to open the file.
-			FileHandle = open(TCHAR_TO_UTF8(*NormalizedFilepath), O_CREAT | O_WRONLY | O_NONBLOCK, S_IRUSR | S_IWUSR | S_IRGRP | S_IWGRP | S_IROTH | S_IWOTH);
+			FileHandle = open(TCHAR_TO_UTF8(*LockPath), O_CREAT | O_WRONLY | O_NONBLOCK, S_IRUSR | S_IWUSR | S_IRGRP | S_IWGRP | S_IROTH | S_IWOTH);
 		}
 
 		// If the file is open, try to lock it, but don't block. If the file is already locked by another process or another thread, blocking here may not honor InTimeout.
diff --git a/Engine/Source/Runtime/Core/Private/Unix/UnixPlatformMutex.cpp b/Engine/Source/Runtime/Core/Private/Unix/UnixPlatformMutex.cpp
--- a/Engine/Source/Runtime/Core/Private/Unix/UnixPlatformMutex.cpp
+++ b/Engine/Source/Runtime/Core/Private/Unix/UnixPlatformMutex.cpp
@@ -22,9 +22,9 @@ FUnixSystemWideMutex::FUnixSystemWideMutex(const FString& InName, FTimespan InTi
 	check(InTimeout >= FTimespan::Zero());
 	check(InTimeout.GetTotalSeconds() < (double)FLT_MAX);
 
-	const FString LockPath = FString(FPlatformProcess::ApplicationSettingsDir()) / InName;
-	FString NormalizedFilepath(LockPath);
-	NormalizedFilepath.ReplaceInline(TEXT("\\"), TEXT("/"));
+	FString SanitizedName(InName);
+	FGenericPlatformMisc::SanitizeSystemWideMutexName(SanitizedName);
+	const FString LockPath = FString(FPlatformProcess::ApplicationSettingsDir()) / SanitizedName;
 	FileHandle = -1;
 
 	double ExpireTimeSecs = FPlatformTime::Seconds() + InTimeout.GetTotalSeconds();
@@ -33,7 +33,7 @@ FUnixSystemWideMutex::FUnixSystemWideMutex(const FString& InName, FTimespan InTi
 		if (FileHandle == -1)
 		{
 			// Try to open the file.
-			FileHandle = open(TCHAR_TO_UTF8(*NormalizedFilepath), O_CREAT | O_WRONLY | O_NONBLOCK, S_IRUSR | S_IWUSR | S_IRGRP | S_IWGRP | S_IROTH | S_IWOTH);
+			FileHandle = open(TCHAR_TO_UTF8(*LockPath), O_CREAT | O_WRONLY | O_NONBLOCK, S_IRUSR | S_IWUSR | S_IRGRP | S_IWGRP | S_IROTH | S_IWOTH);
 		}
 
 		// If the file is open, try to lock it, but don't block. If the file is already locked by another process or another thread, blocking here may not honor InTimeout.
diff --git a/Engine/Source/Runtime/Core/Private/Windows/WindowsPlatformMutex.cpp b/Engine/Source/Runtime/Core/Private/Windows/WindowsPlatformMutex.cpp
--- a/Engine/Source/Runtime/Core/Private/Windows/WindowsPlatformMutex.cpp
+++ b/Engine/Source/Runtime/Core/Private/Windows/WindowsPlatformMutex.cpp
@@ -20,11 +20,11 @@ FWindowsSystemWideMutex::FWindowsSystemWideMutex(const FString& InName, FTimespa
 	check(InTimeout.GetTotalMilliseconds() < (double)0x7FFFFFFE)	// limit timespan to a number of millisec that will fit in a signed int32
 
 	// Disallow backslashes as they aren't allowed in Windows mutex names
-	FString NormalizedMutexName(InName);
-	NormalizedMutexName.ReplaceInline(TEXT("\\"), TEXT("/"));
+	FString SanitizedName(InName);
+	FGenericPlatformMisc::SanitizeSystemWideMutexName(SanitizedName);
 
 	TCHAR MutexName[MAX_PATH] = {};
-	FCString::Strcpy(MutexName, *NormalizedMutexName);
+	FCString::Strcpy(MutexName, *SanitizedName);
 
 	// Attempt to create and take ownership of a named mutex
 	Mutex = CreateMutex(NULL, true, MutexName);
diff --git a/Engine/Source/Runtime/Core/Public/GenericPlatform/GenericPlatformMisc.h b/Engine/Source/Runtime/Core/Public/GenericPlatform/GenericPlatformMisc.h
--- a/Engine/Source/Runtime/Core/Public/GenericPlatform/GenericPlatformMisc.h
+++ b/Engine/Source/Runtime/Core/Public/GenericPlatform/GenericPlatformMisc.h
@@ -2110,6 +2110,12 @@ public:
 		return 0.0f;
 	}
 	
+	/*
+	* Utility to sanitize names provided to FPlatformSystemWideMutex. 
+	* Sanitization is done in-place and will not change the string length.
+	*/
+	static void SanitizeSystemWideMutexName(FString& InName);
+
 #if !UE_BUILD_SHIPPING
 	/**
 	 * Returns any platform specific warning messages we want printed on screen
diff --git a/Engine/Source/Runtime/Core/Public/GenericPlatform/GenericPlatformMutex.h b/Engine/Source/Runtime/Core/Public/GenericPlatform/GenericPlatformMutex.h
--- a/Engine/Source/Runtime/Core/Public/GenericPlatform/GenericPlatformMutex.h
+++ b/Engine/Source/Runtime/Core/Public/GenericPlatform/GenericPlatformMutex.h
@@ -57,6 +57,8 @@ private:
 	MutexType Mutex;
 };
 
+void SanitizeSystemWideMutexName(FString& InName);
+
 /** Platforms that don't need a working FPlatformSystemWideMutex can alias this one. */
 class FPlatformSystemWideMutexNotImplemented
 {
