Fixed TArrayView and TStringView assignment to issue UE_LIFETIMEBOUND warnings.

diff --git a/Engine/Source/Runtime/Core/Public/Concepts/Const.h b/Engine/Source/Runtime/Core/Public/Concepts/Const.h
new file mode 100644
--- /dev/null
+++ b/Engine/Source/Runtime/Core/Public/Concepts/Const.h
@@ -0,0 +1,14 @@
+// Copyright Epic Games, Inc. All Rights Reserved.
+
+#pragma once
+
+#include <type_traits>
+
+namespace UE
+{
+	/**
+	 * Concept which describes a const type.
+	 */
+	template <typename T>
+	concept CConst = std::is_const_v<T>;
+}
diff --git a/Engine/Source/Runtime/Core/Public/Containers/ArrayView.h b/Engine/Source/Runtime/Core/Public/Containers/ArrayView.h
--- a/Engine/Source/Runtime/Core/Public/Containers/ArrayView.h
+++ b/Engine/Source/Runtime/Core/Public/Containers/ArrayView.h
@@ -3,6 +3,7 @@
 #pragma once
 
 #include "CoreTypes.h"
+#include "Concepts/Const.h"
 #include "Containers/ContainersFwd.h" // IWYU pragma: export
 #include "Misc/AssertionMacros.h"
 #include "Misc/ReverseIterate.h"
@@ -252,12 +253,80 @@ public:
 	 * The caller is responsible for ensuring that the view does not outlive the initializer list.
 	 */
 	[[nodiscard]] constexpr TArrayView(std::initializer_list<ElementType> List UE_LIFETIMEBOUND)
+		requires UE::CConst<ElementType>
 		: DataPtr(GetDataHelper(List))
 		, ArrayNum(GetNum(List))
 	{
-		static_assert(std::is_const_v<ElementType>, "Only views of const elements can bind to initializer lists");
 	}
 
+	//////////////////////////////////
+	// Start - assignment operators //
+	//////////////////////////////////
+
+	/**
+	 * We need to explicitly implement assignment operators, even though they're equivalent to copying an implicitly-
+	 * constructed temporary, because UE_LIFETIMEBOUND isn't transitive across implicit construction.
+	 */
+
+	 /**
+	 * Assign from another range
+	 *
+	 * @param Other The source range to copy
+	 */
+	template <
+		typename OtherRangeType,
+		typename CVUnqualifiedOtherRangeType = std::remove_cv_t<std::remove_reference_t<OtherRangeType>>
+		UE_REQUIRES(
+			TAnd<
+				TIsContiguousContainer<CVUnqualifiedOtherRangeType>,
+				TOr<
+					TIsCompatibleRangeType<OtherRangeType>,
+					TIsReinterpretableRangeType<OtherRangeType>
+				>
+			>::Value &&
+			TIsTArrayView_V<CVUnqualifiedOtherRangeType> &&
+			!std::is_same_v<CVUnqualifiedOtherRangeType, TArrayView>
+		)
+	>
+	UE_REWRITE constexpr TArrayView& operator=(OtherRangeType&& Other)
+	{
+		return *this = TArrayView(Forward<OtherRangeType>(Other));
+	}
+	template <
+		typename OtherRangeType,
+		typename CVUnqualifiedOtherRangeType = std::remove_cv_t<std::remove_reference_t<OtherRangeType>>
+		UE_REQUIRES(
+			TAnd<
+				TIsContiguousContainer<CVUnqualifiedOtherRangeType>,
+				TOr<
+					TIsCompatibleRangeType<OtherRangeType>,
+					TIsReinterpretableRangeType<OtherRangeType>
+				>
+			>::Value &&
+			!TIsTArrayView_V<CVUnqualifiedOtherRangeType>
+		)
+	>
+	UE_REWRITE constexpr TArrayView& operator=(OtherRangeType&& Other UE_LIFETIMEBOUND)
+	{
+		return *this = TArrayView(Forward<OtherRangeType>(Other));
+	}
+
+	/**
+	 * Assign from an initializer list.
+	 *
+	 * The caller is responsible for ensuring that the view does not outlive the initializer list.
+	 */
+	UE_REWRITE constexpr TArrayView& operator=(std::initializer_list<ElementType> List UE_LIFETIMEBOUND)
+		requires UE::CConst<ElementType>
+	{
+		return *this = TArrayView(List);
+	}
+
+	////////////////////////////////
+	// End - assignment operators //
+	////////////////////////////////
+
+
 	///////////////////////////////////////////////////
 	// Start - intrusive TOptional<TArrayView> state //
 	///////////////////////////////////////////////////
diff --git a/Engine/Source/Runtime/Core/Public/Containers/StringView.h b/Engine/Source/Runtime/Core/Public/Containers/StringView.h
--- a/Engine/Source/Runtime/Core/Public/Containers/StringView.h
+++ b/Engine/Source/Runtime/Core/Public/Containers/StringView.h
@@ -146,6 +146,41 @@ public:
 	{
 	}
 
+	//////////////////////////////////
+	// Start - assignment operators //
+	//////////////////////////////////
+
+	/**
+	 * We need to explicitly implement assignment operators, even though they're equivalent to copying an implicitly-
+	 * constructed temporary, because UE_LIFETIMEBOUND isn't transitive across implicit construction.
+	 */
+
+	/** Construct a view of the null-terminated string pointed to by InData. */
+	UE_REWRITE constexpr TStringView& operator=(const CharType* InData UE_LIFETIMEBOUND)
+	{
+		return *this = TStringView(InData);
+	}
+
+	/** Construct a view of the null-terminated string pointed to by InData. */
+	template <UE::CCompatibleCharType<CharType> OtherCharType>
+	UE_REWRITE constexpr TStringView& operator=(const OtherCharType* InData UE_LIFETIMEBOUND)
+	{
+		return *this = TStringView(InData);
+	}
+
+	/** Construct a view from a contiguous range of characters, such as FString or TStringBuilder. */
+	template <UE::CNotCVRefTo<ViewType> CharRangeType>
+		// !std::is_pointer is used here to force the pointer constructor to be used when passed an array or pointer.
+		requires (!std::is_pointer_v<std::decay_t<CharRangeType>> && UE::CCompatibleStringViewable<CharRangeType, CharType>)
+	UE_REWRITE constexpr TStringView& operator=(const CharRangeType& InRange UE_LIFETIMEBOUND)
+	{
+		return *this = TStringView(InRange);
+	}
+
+	////////////////////////////////
+	// End - assignment operators //
+	////////////////////////////////
+
 	/** Access the character at the given index in the view. */
 	[[nodiscard]] inline const CharType& operator[](int32 Index) const;
 
