Fix TPagedArray not forwarding allocator to internal Pages array.

#Synced-CL 49969322

diff --git a/Engine/Source/Runtime/Core/Public/Containers/PagedArray.h b/Engine/Source/Runtime/Core/Public/Containers/PagedArray.h
--- a/Engine/Source/Runtime/Core/Public/Containers/PagedArray.h
+++ b/Engine/Source/Runtime/Core/Public/Containers/PagedArray.h
@@ -728,7 +728,7 @@ public:
 	}
 
 private:
-	TArray<PageType> Pages;
+	TArray<PageType, AllocatorType> Pages;
 	SizeType Count = 0;
 
 	[[nodiscard]] UE_FORCEINLINE_HINT ConstIteratorType GetIterator(SizeType Offset) const
