Fixed compile errors when compiling UnrealEditor with portable toolchain

#Synced-CL 50141442

diff --git a/Engine/Plugins/Editor/MacGraphicsSwitching/Source/MacGraphicsSwitching/Private/MacGraphicsSwitchingWidget.cpp b/Engine/Plugins/Editor/MacGraphicsSwitching/Source/MacGraphicsSwitching/Private/MacGraphicsSwitchingWidget.cpp
--- a/Engine/Plugins/Editor/MacGraphicsSwitching/Source/MacGraphicsSwitching/Private/MacGraphicsSwitchingWidget.cpp
+++ b/Engine/Plugins/Editor/MacGraphicsSwitching/Source/MacGraphicsSwitching/Private/MacGraphicsSwitchingWidget.cpp
@@ -3,7 +3,7 @@
 #include "MacGraphicsSwitchingWidget.h"
 #include "MacGraphicsSwitchingModule.h"
 #include "MacGraphicsSwitchingStyle.h"
-#include "Mac/MacPlatformMisc2.h"
+#include "Mac/MacGPUDescriptor.h"
 #include "DetailLayoutBuilder.h"
 #include "Widgets/Input/SComboBox.h"
 #include "Misc/ConfigCacheIni.h"
@@ -20,8 +20,8 @@ SMacGraphicsSwitchingWidget::SMacGraphicsSwitchingWidget()
 	FString DefaultDesc = FString(TEXT("System Default: ")) + FPlatformMisc::GetPrimaryGPUBrand();
 	Renderers.Add(MakeShareable(new FRendererItem(FText::FromString(DefaultDesc), -1, 0)));
 	
-	TArray<FGPUDescriptor> const& GPUs = FPlatformMisc::GetGPUDescriptors();
-	for (FGPUDescriptor const& GPU : GPUs)
+	TArray<FPlatformMisc::FGPUDescriptor> const& GPUs = FPlatformMisc::GetGPUDescriptors();
+	for (FPlatformMisc::FGPUDescriptor const& GPU : GPUs)
 	{
 		FString RendererDesc = FString::Printf(TEXT("%d: %s"), GPU.GPUIndex, *FString(GPU.GPUName));
 		
@@ -87,8 +87,8 @@ TSharedRef<SWidget> SMacGraphicsSwitchingWidget::OnGenerateWidget( TSharedPtr<FR
 
 void SMacGraphicsSwitchingWidget::OnComboBoxOpening()
 {
-	TArray<FGPUDescriptor> const& GPUs = FPlatformMisc::GetGPUDescriptors();
-	for (FGPUDescriptor const& GPU : GPUs)
+	TArray<FPlatformMisc::FGPUDescriptor> const& GPUs = FPlatformMisc::GetGPUDescriptors();
+	for (FPlatformMisc::FGPUDescriptor const& GPU : GPUs)
 	{
 		bool bFound = false;
 		for (TSharedPtr<FRendererItem> const& Item : Renderers)
@@ -112,7 +112,7 @@ void SMacGraphicsSwitchingWidget::OnComboBoxOpening()
 	for (auto It = Renderers.CreateIterator(); It; ++It)
 	{
 		bool bFound = false;
-		for (FGPUDescriptor const& GPU : GPUs)
+		for (FPlatformMisc::FGPUDescriptor const& GPU : GPUs)
 		{
 			if ((*It)->RegistryID == GPU.RegistryID)
 			{
diff --git a/Engine/Plugins/Fab/Source/Fab/Private/Importers/BuildPatchInstallerLibHelper.h b/Engine/Plugins/Fab/Source/Fab/Private/Importers/BuildPatchInstallerLibHelper.h
--- a/Engine/Plugins/Fab/Source/Fab/Private/Importers/BuildPatchInstallerLibHelper.h
+++ b/Engine/Plugins/Fab/Source/Fab/Private/Importers/BuildPatchInstallerLibHelper.h
@@ -12,7 +12,7 @@
 #include "Misc/Paths.h"
 
 #if PLATFORM_MAC
-#include "Mac/MacPlatformMisc2.h"
+#include "Mac/MacPlatformMiscEx.h"
 #ifndef RTLD_DEEPBIND
 #define RTLD_DEEPBIND 0
 #endif
diff --git a/Engine/Source/Editor/UnrealEd/Private/Factories/TTFontImport.cpp b/Engine/Source/Editor/UnrealEd/Private/Factories/TTFontImport.cpp
--- a/Engine/Source/Editor/UnrealEd/Private/Factories/TTFontImport.cpp
+++ b/Engine/Source/Editor/UnrealEd/Private/Factories/TTFontImport.cpp
@@ -27,7 +27,9 @@
 #endif // WITH_FREETYPE
 
 #if PLATFORM_MAC
-#include "Mac/MacSystemIncludes.h"
+#include "Apple/PreAppleSystemHeaders.h"
+#include <CoreText/CoreText.h>
+#include "Apple/PostAppleSystemHeaders.h"
 #endif
 
 #if PLATFORM_WINDOWS
diff --git a/Engine/Source/Runtime/Apple/MetalRHI/Private/MetalDevice.cpp b/Engine/Source/Runtime/Apple/MetalRHI/Private/MetalDevice.cpp
--- a/Engine/Source/Runtime/Apple/MetalRHI/Private/MetalDevice.cpp
+++ b/Engine/Source/Runtime/Apple/MetalRHI/Private/MetalDevice.cpp
@@ -24,6 +24,10 @@
 #include "MetalBindlessDescriptors.h"
 #include "MetalTempAllocator.h"
 
+#if PLATFORM_MAC
+#include "Mac/MacGPUDescriptor.h"
+#endif
+
 int32 GMetalSupportsIntermediateBackBuffer = 0;
 static FAutoConsoleVariableRef CVarMetalSupportsIntermediateBackBuffer(
 	TEXT("rhi.Metal.SupportsIntermediateBackBuffer"),
diff --git a/Engine/Source/Runtime/Apple/MetalRHI/Private/MetalRHI.cpp b/Engine/Source/Runtime/Apple/MetalRHI/Private/MetalRHI.cpp
--- a/Engine/Source/Runtime/Apple/MetalRHI/Private/MetalRHI.cpp
+++ b/Engine/Source/Runtime/Apple/MetalRHI/Private/MetalRHI.cpp
@@ -26,6 +26,8 @@
 #include "IOS/IOSAppDelegate.h"
 #elif PLATFORM_MAC
 #include "Mac/MacApplication.h"
+#include "Mac/MacGPUDescriptor.h"
+#include "Mac/MacPlatformMiscEx.h"
 #include "HAL/PlatformApplicationMisc.h"
 #include "GenericPlatform/GenericPlatformFile.h"
 #endif
@@ -1381,7 +1383,7 @@ void FMetalRHICommandContext::RHIEndBreadcrumbGPU(FRHIBreadcrumbNode* Breadcrumb
 void FMetalDynamicRHI::RHIGetSupportedResolution( uint32 &Width, uint32 &Height )
 {
 #if PLATFORM_MAC
-	CGDisplayModeRef DisplayMode = FPlatformApplicationMisc::GetSupportedDisplayMode(kCGDirectMainDisplay, Width, Height);
+	CGDisplayModeRef DisplayMode = FMacPlatformMiscEx::GetSupportedDisplayMode(kCGDirectMainDisplay, Width, Height);
 	if (DisplayMode)
 	{
 		Width = CGDisplayModeGetWidth(DisplayMode);
diff --git a/Engine/Source/Runtime/ApplicationCore/Private/Mac/MacPlatformApplicationMisc.cpp b/Engine/Source/Runtime/ApplicationCore/Private/Mac/MacPlatformApplicationMisc.cpp
--- a/Engine/Source/Runtime/ApplicationCore/Private/Mac/MacPlatformApplicationMisc.cpp
+++ b/Engine/Source/Runtime/ApplicationCore/Private/Mac/MacPlatformApplicationMisc.cpp
@@ -522,38 +522,6 @@ float FMacPlatformApplicationMisc::GetDPIScaleFactorAtPoint(float X, float Y)
 	return 1.0f;
 }
 
-CGDisplayModeRef FMacPlatformApplicationMisc::GetSupportedDisplayMode(CGDirectDisplayID DisplayID, uint32 Width, uint32 Height)
-{
-	CGDisplayModeRef BestMatchingMode = nullptr;
-	uint32 BestWidth = 0;
-	uint32 BestHeight = 0;
-
-	CFArrayRef AllModes = CGDisplayCopyAllDisplayModes(DisplayID, nullptr);
-	if (AllModes)
-	{
-		const int32 NumModes = CFArrayGetCount(AllModes);
-		for (int32 Index = 0; Index < NumModes; Index++)
-		{
-			CGDisplayModeRef Mode = (CGDisplayModeRef)CFArrayGetValueAtIndex(AllModes, Index);
-			const int32 ModeWidth = (int32)CGDisplayModeGetWidth(Mode);
-			const int32 ModeHeight = (int32)CGDisplayModeGetHeight(Mode);
-
-			const bool bIsEqualOrBetterWidth = FMath::Abs((int32)ModeWidth - (int32)Width) <= FMath::Abs((int32)BestWidth - (int32)Width);
-			const bool bIsEqualOrBetterHeight = FMath::Abs((int32)ModeHeight - (int32)Height) <= FMath::Abs((int32)BestHeight - (int32)Height);
-			if (!BestMatchingMode || (bIsEqualOrBetterWidth && bIsEqualOrBetterHeight))
-			{
-				BestWidth = ModeWidth;
-				BestHeight = ModeHeight;
-				BestMatchingMode = Mode;
-			}
-		}
-		BestMatchingMode = CGDisplayModeRetain(BestMatchingMode);
-		CFRelease(AllModes);
-	}
-
-	return BestMatchingMode;
-}
-
 void FMacPlatformApplicationMisc::PumpMessages(bool bFromMainLoop)
 {
 	if( bFromMainLoop )
diff --git a/Engine/Source/Runtime/ApplicationCore/Private/Mac/MacWindow.cpp b/Engine/Source/Runtime/ApplicationCore/Private/Mac/MacWindow.cpp
--- a/Engine/Source/Runtime/ApplicationCore/Private/Mac/MacWindow.cpp
+++ b/Engine/Source/Runtime/ApplicationCore/Private/Mac/MacWindow.cpp
@@ -9,6 +9,7 @@
 #include "Mac/CocoaTextView.h"
 #include "Mac/CocoaThread.h"
 #include "Mac/MacPlatformApplicationMisc.h"
+#include "Mac/MacPlatformMiscEx.h"
 
 #include "HAL/PlatformProcess.h"
 
@@ -701,7 +702,7 @@ void FMacWindow::ApplySizeAndModeChanges(int32 X, int32 Y, int32 Width, int32 He
 
 		if (WindowedModeSavedState.CapturedDisplayID != kCGNullDirectDisplay)
 		{
-			CGDisplayModeRef DisplayMode = FPlatformApplicationMisc::GetSupportedDisplayMode(WindowedModeSavedState.CapturedDisplayID, Width, Height);
+			CGDisplayModeRef DisplayMode = FMacPlatformMiscEx::GetSupportedDisplayMode(WindowedModeSavedState.CapturedDisplayID, Width, Height);
 			MainThreadCall(^{
 				CGDisplaySetDisplayMode(WindowedModeSavedState.CapturedDisplayID, DisplayMode, nullptr);
 			}, true, UnrealResizeEventMode);
diff --git a/Engine/Source/Runtime/ApplicationCore/Public/Mac/MacPlatformApplicationMisc.h b/Engine/Source/Runtime/ApplicationCore/Public/Mac/MacPlatformApplicationMisc.h
--- a/Engine/Source/Runtime/ApplicationCore/Public/Mac/MacPlatformApplicationMisc.h
+++ b/Engine/Source/Runtime/ApplicationCore/Public/Mac/MacPlatformApplicationMisc.h
@@ -4,6 +4,10 @@
 
 #include "GenericPlatform/GenericPlatformApplicationMisc.h"
 
+#ifdef __OBJC__
+@protocol NSObject;
+#endif
+
 typedef void (*MacApplicationExternalCb)(void);
 typedef void (*MacApplicationExternalCbOneBool)(bool);
 
@@ -32,7 +36,6 @@ struct APPLICATIONCORE_API FMacPlatformApplicationMisc : public FGenericPlatform
 	static void PumpMessages(bool bFromMainLoop);
 	static void ClipboardCopy(const TCHAR* Str);
 	static void ClipboardPaste(class FString& Dest);
-	static CGDisplayModeRef GetSupportedDisplayMode(CGDirectDisplayID DisplayID, uint32 Width, uint32 Height);
 	
 	// MAC ONLY
 
diff --git a/Engine/Source/Runtime/Core/Private/Apple/ApplePlatformMisc.cpp b/Engine/Source/Runtime/Core/Private/Apple/ApplePlatformMisc.cpp
--- a/Engine/Source/Runtime/Core/Private/Apple/ApplePlatformMisc.cpp
+++ b/Engine/Source/Runtime/Core/Private/Apple/ApplePlatformMisc.cpp
@@ -6,7 +6,7 @@
 
 #include "Apple/ApplePlatformMisc.h"
 #if PLATFORM_MAC
-#include "Mac/MacSystemIncludes.h"
+#include "Mac/MacPlatformMiscEx.h"
 #elif PLATFORM_IOS
 #include "IOS/IOSSystemIncludes.h"
 #endif
@@ -357,7 +357,7 @@ bool FApplePlatformMisc::IsOSAtLeastVersion(const uint32 MacOSVersion[3], const
 {
 	NSOperatingSystemVersion CurrentSystemVersion =
 #if PLATFORM_MAC
-	FMacPlatformMisc::GetNSOperatingSystemVersion();
+	FMacPlatformMiscEx::GetNSOperatingSystemVersion();
 #else
 	[NSProcessInfo processInfo].operatingSystemVersion;
 #endif
diff --git a/Engine/Source/Runtime/Core/Private/Mac/MacPlatformMisc.cpp b/Engine/Source/Runtime/Core/Private/Mac/MacPlatformMisc.cpp
--- a/Engine/Source/Runtime/Core/Private/Mac/MacPlatformMisc.cpp
+++ b/Engine/Source/Runtime/Core/Private/Mac/MacPlatformMisc.cpp
@@ -13,7 +13,9 @@
 #include "Misc/EngineVersion.h"
 #include "Mac/MacMallocZone.h"
 #include "Apple/ApplePlatformSymbolication.h"
+#include "Mac/MacGPUDescriptor.h"
 #include "Mac/MacPlatformCrashContext.h"
+#include "Mac/MacPlatformMiscEx.h"
 #include "GenericPlatform/GenericPlatformDriver.h"
 #include "HAL/ThreadHeartBeat.h"
 #include "HAL/PlatformOutputDevices.h"
@@ -258,7 +260,7 @@ struct FMacApplicationInfo
 		
 		MacPlatformGetOSVersion(OSBuild);
 
-		OSXVersion = FMacPlatformMisc::GetNSOperatingSystemVersion();
+		OSXVersion = FMacPlatformMiscEx::GetNSOperatingSystemVersion();
 
 		XcodeVersion.majorVersion = XcodeVersion.minorVersion = XcodeVersion.patchVersion = 0;
 
@@ -841,7 +843,7 @@ void FMacPlatformMisc::NormalizePath(FStringBuilderBase& InPath)
 }
 
 template<typename T>
-FMacPlatformMisc::FGPUDescriptorCommon<T>::~FGPUDescriptorCommon()
+FMacGPUDescriptorCommon<T>::~FMacGPUDescriptorCommon()
 {
 	[GPUName release];
 	[GPUMetalBundle release];
@@ -850,7 +852,7 @@ FMacPlatformMisc::FGPUDescriptorCommon<T>::~FGPUDescriptorCommon()
 }
 
 template<typename T>
-void FMacPlatformMisc::FGPUDescriptorCommon<T>::CopyFrom(FMacPlatformMisc::FGPUDescriptorCommon<T> const& Other)
+void FMacGPUDescriptorCommon<T>::CopyFrom(FMacGPUDescriptorCommon<T> const& Other)
 {
 	if (this != &Other)
 	{
@@ -880,25 +882,25 @@ void FMacPlatformMisc::FGPUDescriptorCommon<T>::CopyFrom(FMacPlatformMisc::FGPUD
 }
 
 template<typename T>
-FMacPlatformMisc::FGPUDescriptorCommon<T>& FMacPlatformMisc::FGPUDescriptorCommon<T>::operator=(FGPUDescriptorCommon<T> const& Other)
+FMacGPUDescriptorCommon<T>& FMacGPUDescriptorCommon<T>::operator=(FMacGPUDescriptorCommon<T> const& Other)
 {
 	CopyFrom(Other);
 	return *this;
 }
 
 template<typename T>
-TMap<FString, float> FMacPlatformMisc::FGPUDescriptorCommon<T>::GetPerformanceStatistics() const
+TMap<FString, float> FMacGPUDescriptorCommon<T>::GetPerformanceStatistics() const
 {
 	return static_cast<T const*>(this)->GetPerformanceStatisticsImpl();
 }
 
 #if PLATFORM_MAC_X86
-FMacPlatformMisc::FGPUDescriptorX86_64::FGPUDescriptorX86_64(FGPUDescriptorX86_64 const& Other)
+FMacGPUDescriptor::FMacGPUDescriptor(FMacGPUDescriptor const& Other)
 {
 	CopyFrom(Other);
 }
 
-FMacPlatformMisc::FGPUDescriptorX86_64::~FGPUDescriptorX86_64()
+FMacGPUDescriptor::~FMacGPUDescriptor()
 {
 	if (PCIDevice)
 	{
@@ -906,9 +908,9 @@ FMacPlatformMisc::FGPUDescriptorX86_64::~FGPUDescriptorX86_64()
 	}
 }
 
-void FMacPlatformMisc::FGPUDescriptorX86_64::CopyFromImpl(FGPUDescriptorCommon<FMacPlatformMisc::FGPUDescriptorX86_64> const& Other)
+void FMacGPUDescriptor::CopyFromImpl(FMacGPUDescriptorCommon<FMacGPUDescriptor> const& Other)
 {
-	FMacPlatformMisc::FGPUDescriptorX86_64 const& ConcreteOther = static_cast<FMacPlatformMisc::FGPUDescriptorX86_64 const&>(Other);
+	FMacGPUDescriptor const& ConcreteOther = static_cast<FMacGPUDescriptor const&>(Other);
 
 	RegistryID = ConcreteOther.RegistryID;
 
@@ -923,7 +925,7 @@ void FMacPlatformMisc::FGPUDescriptorX86_64::CopyFromImpl(FGPUDescriptorCommon<F
 	PCIDevice = ConcreteOther.PCIDevice;
 }
 
-TMap<FString, float> FMacPlatformMisc::FGPUDescriptorX86_64::GetPerformanceStatisticsImpl() const
+TMap<FString, float> FMacGPUDescriptor::GetPerformanceStatisticsImpl() const
 {
 	SCOPED_AUTORELEASE_POOL;
 
@@ -949,23 +951,23 @@ TMap<FString, float> FMacPlatformMisc::FGPUDescriptorX86_64::GetPerformanceStati
 }
 
 #elif PLATFORM_MAC_ARM64
-FMacPlatformMisc::FGPUDescriptorARM64::FGPUDescriptorARM64(FGPUDescriptorARM64 const& Other)
+FMacGPUDescriptor::FMacGPUDescriptor(FMacGPUDescriptor const& Other)
 {
 	CopyFrom(Other);
 }
 
-FMacPlatformMisc::FGPUDescriptorARM64::~FGPUDescriptorARM64()
+FMacGPUDescriptor::~FMacGPUDescriptor()
 {
 }
 
-void FMacPlatformMisc::FGPUDescriptorARM64::CopyFromImpl(FGPUDescriptorCommon<FMacPlatformMisc::FGPUDescriptorARM64> const& Other)
+void FMacGPUDescriptor::CopyFromImpl(FMacGPUDescriptorCommon<FMacGPUDescriptor> const& Other)
 {
-	FMacPlatformMisc::FGPUDescriptorARM64 const& ConcreteOther = static_cast<FMacPlatformMisc::FGPUDescriptorARM64 const&>(Other);
+	FMacGPUDescriptor const& ConcreteOther = static_cast<FMacGPUDescriptor const&>(Other);
 
 	RegistryID = ConcreteOther.RegistryID;
 }
 
-TMap<FString, float> FMacPlatformMisc::FGPUDescriptorARM64::GetPerformanceStatisticsImpl() const
+TMap<FString, float> FMacGPUDescriptor::GetPerformanceStatisticsImpl() const
 {
 	return TMap<FString, float>();
 }
@@ -978,11 +980,11 @@ TMap<FString, float> FMacPlatformMisc::FGPUDescriptorARM64::GetPerformanceStatis
 class FMacPlatformGPUManager
 {
 	FCriticalSection Mutex;
-	TArray<FMacPlatformMisc::FGPUDescriptor> CurrentGPUs;
-	TArray<FMacPlatformMisc::FGPUDescriptor> UpdatedGPUs;
+	TArray<FMacGPUDescriptor> CurrentGPUs;
+	TArray<FMacGPUDescriptor> UpdatedGPUs;
 	TAtomic<bool> bRequiresUpdate;
 
-	void InitializeDescriptorFromDeviceEntryM(FMacPlatformMisc::FGPUDescriptor& Desc, io_registry_entry_t ServiceEntry, CFMutableDictionaryRef ServiceInfo)
+	void InitializeDescriptorFromDeviceEntryM(FMacGPUDescriptor& Desc, io_registry_entry_t ServiceEntry, CFMutableDictionaryRef ServiceInfo)
 	{
 		static CFStringRef IOMatchCategoryRef	= CFSTR("IOMatchCategory");
 		static CFStringRef IOAcceleratorRef		= CFSTR("IOAccelerator");
@@ -1077,7 +1079,7 @@ class FMacPlatformGPUManager
 	}
 
 #if PLATFORM_MAC_X86
-	void InitializeDescriptorFromDeviceEntry(FMacPlatformMisc::FGPUDescriptor& Desc, io_registry_entry_t ServiceEntry, CFMutableDictionaryRef ServiceInfo)
+	void InitializeDescriptorFromDeviceEntry(FMacGPUDescriptor& Desc, io_registry_entry_t ServiceEntry, CFMutableDictionaryRef ServiceInfo)
 	{
 		IOObjectRetain(ServiceEntry);
 		Desc.PCIDevice = (uint32)ServiceEntry;
@@ -1234,7 +1236,7 @@ public:
 							const char* ClassCodeValue = reinterpret_cast<const char*>(CFDataGetBytePtr(ClassCode));
 							if (!strncasecmp(ClassCodeValue, "sgx", 3))
 							{
-								FMacPlatformMisc::FGPUDescriptor Desc;
+								FMacGPUDescriptor Desc;
 #if PLATFORM_MAC_X86
 								// Default initialized to 0 on ARM64, but initialized here for X86 due to
 								// 4.26.0 compatibility requirements.
@@ -1256,7 +1258,7 @@ public:
 							// GPUs are class-code 0x30000 || 0x38000
 							if (ClassCodeValue && (*ClassCodeValue == 0x30000 || *ClassCodeValue == 0x38000))
 							{
-								FMacPlatformMisc::FGPUDescriptor Desc;
+								FMacGPUDescriptor Desc;
 
 								InitializeDescriptorFromDeviceEntry(Desc, ServiceEntry, ServiceInfo);
 
@@ -1282,7 +1284,7 @@ public:
 		UpdatedGPUs = CurrentGPUs;
 	}
 	
-	TArray<FMacPlatformMisc::FGPUDescriptor> const& GetCurrentGPUs()
+	TArray<FMacGPUDescriptor> const& GetCurrentGPUs()
 	{
 		if (bRequiresUpdate)
 		{
@@ -1325,7 +1327,7 @@ public:
 										{
 											FScopeLock Lock(&Mutex);
 											
-											FMacPlatformMisc::FGPUDescriptor Desc;
+											FMacGPUDescriptor Desc;
 											
                                             InitializeDescriptorFromDeviceEntry(Desc, ServiceEntry, ServiceInfo);
 											
@@ -1357,7 +1359,7 @@ public:
 				FScopeLock Lock(&Mutex);
 				for (int32 i = 0; i < UpdatedGPUs.Num(); i++)
 				{
-					FMacPlatformMisc::FGPUDescriptor& Desc = UpdatedGPUs[i];
+					FMacGPUDescriptor& Desc = UpdatedGPUs[i];
 					if (Desc.RegistryID == DeviceRegistryID)
 					{
 						if (Desc.GPUIndex == GMacExplicitRendererID)
@@ -1370,7 +1372,7 @@ public:
 				}
 				for (int32 i = 0; i < UpdatedGPUs.Num(); i++)
 				{
-					FMacPlatformMisc::FGPUDescriptor& Desc = UpdatedGPUs[i];
+					FMacGPUDescriptor& Desc = UpdatedGPUs[i];
 					Desc.GPUIndex = (uint32)i;
 				}
 				bRequiresUpdate = true;
@@ -1390,7 +1392,7 @@ void FMacPlatformMisc::GPUChangeNotification(uint64_t DeviceRegistryID, EMacGPUN
 	FMacPlatformGPUManager::Get().Notify(DeviceRegistryID, Notification);
 }
 
-TArray<FMacPlatformMisc::FGPUDescriptor> const& FMacPlatformMisc::GetGPUDescriptors()
+TArray<FMacGPUDescriptor> const& FMacPlatformMisc::GetGPUDescriptors()
 {
 	return FMacPlatformGPUManager::Get().GetCurrentGPUs();
 }
@@ -1417,11 +1419,11 @@ FString FMacPlatformMisc::GetPrimaryGPUBrand()
 	static FString PrimaryGPU;
 	if(PrimaryGPU.IsEmpty())
 	{
-		TArray<FMacPlatformMisc::FGPUDescriptor> const& GPUs = GetGPUDescriptors();
+		TArray<FMacGPUDescriptor> const& GPUs = GetGPUDescriptors();
 		
 		if(GPUs.Num() > 1)
 		{
-			for(FMacPlatformMisc::FGPUDescriptor const& GPU : GPUs)
+			for(FMacGPUDescriptor const& GPU : GPUs)
 			{
 				if(!GPU.GPUHeadless && GPU.GPUVendorId != 0x8086)
 				{
@@ -1476,9 +1478,9 @@ FGPUDriverInfo FMacPlatformMisc::GetGPUDriverInfo(const FString& DeviceDescripti
 
 	FGPUDriverInfo Info;
 	TArray<FString> NameComponents;
-	TArray<FMacPlatformMisc::FGPUDescriptor> const& GPUs = GetGPUDescriptors();
+	TArray<FMacGPUDescriptor> const& GPUs = GetGPUDescriptors();
 	
-	for(FMacPlatformMisc::FGPUDescriptor const& GPU : GPUs)
+	for(FMacGPUDescriptor const& GPU : GPUs)
 	{
 		NameComponents.Empty();
 		bool bMatchesName = FString(GPU.GPUName).TrimStart().ParseIntoArray(NameComponents, TEXT(" ")) > 0;
@@ -1666,7 +1668,7 @@ FString FMacPlatformMisc::GetOSVersion()
 	return OSVersion;
 }
 
-NSOperatingSystemVersion FMacPlatformMisc::GetNSOperatingSystemVersion()
+NSOperatingSystemVersion FMacPlatformMiscEx::GetNSOperatingSystemVersion()
 {
 	FString OSVersionString;
 	NSOperatingSystemVersion SystemVersion = {};
@@ -1821,7 +1823,7 @@ bool FMacPlatformMisc::IsRunningOnBattery()
 	return GMacAppInfo.RunningOnBattery;
 }
 
-int32 FMacPlatformMisc::VersionCompare(const NSOperatingSystemVersion& VersionA, const NSOperatingSystemVersion& VersionB)
+int32 FMacPlatformMiscEx::VersionCompare(const NSOperatingSystemVersion& VersionA, const NSOperatingSystemVersion& VersionB)
 {
 	NSInteger ValuesA[3] = { VersionA.majorVersion, VersionA.minorVersion, VersionA.patchVersion };
 	NSInteger ValuesB[3] = { VersionB.majorVersion, VersionB.minorVersion, VersionB.patchVersion };
@@ -1843,7 +1845,7 @@ int32 FMacPlatformMisc::VersionCompare(const NSOperatingSystemVersion& VersionA,
 
 int32 FMacPlatformMisc::MacOSXVersionCompare(uint8 Major, uint8 Minor, uint8 Revision)
 {
-	return VersionCompare(GMacAppInfo.OSXVersion, {Major, Minor, Revision});
+	return FMacPlatformMiscEx::VersionCompare(GMacAppInfo.OSXVersion, {Major, Minor, Revision});
 }
 
 FString FMacPlatformMisc::GetOperatingSystemId()
@@ -1871,7 +1873,7 @@ FString FMacPlatformMisc::GetXcodePath()
 
 int32 FMacPlatformMisc::XCodeVersionCompare(uint8 Major, uint8 Minor, uint8 Revision)
 {
-	return VersionCompare(GMacAppInfo.XcodeVersion, {Major, Minor, Revision});
+	return FMacPlatformMiscEx::VersionCompare(GMacAppInfo.XcodeVersion, {Major, Minor, Revision});
 }
 
 bool FMacPlatformMisc::IsSupportedXcodeVersionInstalled()
@@ -1960,7 +1962,7 @@ bool FMacPlatformMisc::IsRunningOnRecommendedMinSpecHardware()
 }
 #endif // WITH_EDITOR
 
-CGDisplayModeRef FMacPlatformMisc::GetSupportedDisplayMode(CGDirectDisplayID DisplayID, uint32 Width, uint32 Height)
+CGDisplayModeRef FMacPlatformMiscEx::GetSupportedDisplayMode(CGDirectDisplayID DisplayID, uint32 Width, uint32 Height)
 {
 	CGDisplayModeRef BestMatchingMode = nullptr;
 	uint32 BestWidth = 0;
@@ -2732,11 +2734,11 @@ bool FMacPlatformMisc::HasPlatformFeature(const TCHAR* FeatureName)
                 void* DLLHandle = FPlatformProcess::GetDllHandle(TEXT("/System/Library/Frameworks/Metal.framework/Metal"));
                 if (DLLHandle)
                 {
-                    TArray<FMacPlatformMisc::FGPUDescriptor> const& GPUs = FPlatformMisc::GetGPUDescriptors();
+                    TArray<FMacGPUDescriptor> const& GPUs = FPlatformMisc::GetGPUDescriptors();
                     
                     if (GPUs.Num() > 0)
                     {
-                        for (FMacPlatformMisc::FGPUDescriptor const& GPU : GPUs)
+                        for (FMacGPUDescriptor const& GPU : GPUs)
                         {
                             NSString* GPUName = [GPU.GPUName length] > 0 ? GPU.GPUName : @"Unnamed GPU";
                             NSString* GPUMetalBundle = [GPU.GPUMetalBundle length] > 0 ? GPU.GPUMetalBundle : @"null";
@@ -2900,10 +2902,10 @@ void FMacPlatformMisc::UpdateDriverMonitorStatistics(int32 DeviceIndex)
 {
 	if (DeviceIndex >= 0)
 	{
-		TArray<FMacPlatformMisc::FGPUDescriptor> const& GPUs = FPlatformMisc::GetGPUDescriptors();
+		TArray<FMacGPUDescriptor> const& GPUs = FPlatformMisc::GetGPUDescriptors();
 		if (DeviceIndex < GPUs.Num())
 		{
-			FMacPlatformMisc::FGPUDescriptor const& GPU = GPUs[DeviceIndex];
+			FMacGPUDescriptor const& GPU = GPUs[DeviceIndex];
 			TMap<FString, float> Stats = GPU.GetPerformanceStatistics();
 			
 			float DeviceUtil = GetMacGPUStat<float>(Stats, TEXT("Device Utilization %"));
diff --git a/Engine/Source/Runtime/Core/Public/Mac/MacGPUDescriptor.h b/Engine/Source/Runtime/Core/Public/Mac/MacGPUDescriptor.h
new file mode 100644
--- /dev/null
+++ b/Engine/Source/Runtime/Core/Public/Mac/MacGPUDescriptor.h
@@ -0,0 +1,87 @@
+// Copyright Epic Games, Inc. All Rights Reserved.
+
+/*=============================================================================================
+	MacPlatformMiscEx.h: Mac platform additional misc functions that requires sdk depdendencies
+==============================================================================================*/
+
+#pragma once
+
+#include "Containers/Map.h"
+#include "Mac/MacSystemIncludes.h"
+
+/** Common descriptor of each GPU in the OS that provides stock details about the GPU that are innaccessible from the higher-level rendering APIs and provides a direct link to the GPU in the IORegistry. */
+template<typename T>
+class FMacGPUDescriptorCommon
+{
+	public:
+		virtual ~FMacGPUDescriptorCommon();
+
+		FMacGPUDescriptorCommon& operator=(FMacGPUDescriptorCommon const& Other);
+
+		TMap<FString, float> GetPerformanceStatistics() const;
+
+	protected:
+		FMacGPUDescriptorCommon() = default;
+		FMacGPUDescriptorCommon(FMacGPUDescriptorCommon const& Other) = delete;
+		FMacGPUDescriptorCommon(FMacGPUDescriptorCommon&& Other) = delete;
+
+		void CopyFrom(FMacGPUDescriptorCommon const& Other);
+
+	public:
+		NSString*	GPUName			= nil;
+		NSString*	GPUMetalBundle	= nil;
+		NSString*	GPUOpenGLBundle	= nil;
+		NSString*	GPUBundleID		= nil;
+		uint32		GPUVendorId		= 0;
+		uint32		GPUDeviceId		= 0;
+		uint32		GPUMemoryMB		= 0;
+		uint32		GPUIndex		= 0;
+
+		bool GPUHeadless			= false;
+};
+
+#if PLATFORM_MAC_X86
+	/** Intel architecture descriptor of each GPU in the OS that provides stock details about the GPU that are innaccessible from the higher-level rendering APIs and provides a direct link to the GPU in the IORegistry. */
+	class FMacGPUDescriptor : public FMacGPUDescriptorCommon<FMacGPUDescriptor>
+	{
+		friend class FMacGPUDescriptorCommon<FMacGPUDescriptor>;
+
+		public:
+			FMacGPUDescriptor() = default;
+			FMacGPUDescriptor(FMacGPUDescriptor const& Other);
+
+			virtual ~FMacGPUDescriptor();
+
+		protected:
+			void CopyFromImpl(FMacGPUDescriptorCommon const& Other);
+			TMap<FString, float> GetPerformanceStatisticsImpl() const;
+
+		public:
+			uint64 RegistryID;
+			uint32 PCIDevice; // This is really an io_registry_entry_t which is a mach port name
+
+	};
+
+#elif PLATFORM_MAC_ARM64
+	class FMacGPUDescriptor : public FMacGPUDescriptorCommon<FMacGPUDescriptor>
+	{
+		friend class FMacGPUDescriptorCommon<FMacGPUDescriptor>;
+
+		public:
+			FMacGPUDescriptor() = default;
+			FMacGPUDescriptor(FMacGPUDescriptor const& Other);
+
+			virtual ~FMacGPUDescriptor();
+
+		protected:
+			void CopyFromImpl(FMacGPUDescriptorCommon const& Other);
+			TMap<FString, float> GetPerformanceStatisticsImpl() const;
+
+		public:
+			uint64 RegistryID = 0;
+
+	};
+
+#else
+	#error "Undefined Mac platform"
+#endif
diff --git a/Engine/Source/Runtime/Core/Public/Mac/MacPlatformMisc.h b/Engine/Source/Runtime/Core/Public/Mac/MacPlatformMisc.h
--- a/Engine/Source/Runtime/Core/Public/Mac/MacPlatformMisc.h
+++ b/Engine/Source/Runtime/Core/Public/Mac/MacPlatformMisc.h
@@ -6,9 +6,18 @@
 
 #pragma once
 #include "GenericPlatform/GenericPlatformMisc.h"
+#if UE_ENABLE_INCLUDE_ORDER_DEPRECATED_IN_5_8
 #include "Mac/MacSystemIncludes.h"
+#endif // UE_ENABLE_INCLUDE_ORDER_DEPRECATED_IN_5_8
+
+#ifdef __OBJC__
+@class Protocol;
+#endif
+
 #include "Apple/ApplePlatformMisc.h"
 
+class FMacGPUDescriptor;
+
 /**
 * Mac implementation of the misc OS functions
 **/
@@ -41,7 +50,6 @@ struct CORE_API FMacPlatformMisc : public FApplePlatformMisc
 	static struct FGPUDriverInfo GetGPUDriverInfo(const FString& DeviceDescription, bool bVerbose = true);
 	static void GetOSVersions( FString& out_OSVersionLabel, FString& out_OSSubVersionLabel );
 	static FString GetOSVersion();
-	static NSOperatingSystemVersion GetNSOperatingSystemVersion();
 	static bool HasPlatformFeature(const TCHAR* FeatureName);
 	static bool GetDiskTotalAndFreeSpace(const FString& InPath, uint64& TotalNumberOfBytes, uint64& NumberOfFreeBytes);
 	static bool HasSeparateChannelForDebugOutput();
@@ -101,11 +109,6 @@ struct CORE_API FMacPlatformMisc : public FApplePlatformMisc
 
 	static IPlatformChunkInstall* GetPlatformChunkInstall();
 
-	/**
-	 * Returns if A < B returns -1, if A > B returns 1, else A == B and returns 0.
-	 */
-	static int32 VersionCompare(const NSOperatingSystemVersion& VersionA, const NSOperatingSystemVersion& VersionB);
-
 	/**
 	 * Returns if current < target returns 1, if current > target returns 1, else current == target and returns 0.
 	 */
@@ -131,83 +134,6 @@ struct CORE_API FMacPlatformMisc : public FApplePlatformMisc
 
 	static void MergeDefaultArgumentsIntoCommandLine(FString& CommandLine, FString DefaultArguments);
 	
-	/** Common descriptor of each GPU in the OS that provides stock details about the GPU that are innaccessible from the higher-level rendering APIs and provides a direct link to the GPU in the IORegistry. */
-	template<typename T>
-	class FGPUDescriptorCommon
-	{
-		public:
-			virtual ~FGPUDescriptorCommon();
-
-			FGPUDescriptorCommon& operator=(FGPUDescriptorCommon const& Other);
-
-			TMap<FString, float> GetPerformanceStatistics() const;
-
-		protected:
-			FGPUDescriptorCommon() = default;
-			FGPUDescriptorCommon(FGPUDescriptorCommon const& Other) = delete;
-			FGPUDescriptorCommon(FGPUDescriptorCommon&& Other) = delete;
-
-			void CopyFrom(FGPUDescriptorCommon const& Other);
-
-		public:
-			NSString*	GPUName			= nil;
-			NSString*	GPUMetalBundle	= nil;
-			NSString*	GPUOpenGLBundle	= nil;
-			NSString*	GPUBundleID		= nil;
-			uint32		GPUVendorId		= 0;
-			uint32		GPUDeviceId		= 0;
-			uint32		GPUMemoryMB		= 0;
-			uint32		GPUIndex		= 0;
-
-			bool GPUHeadless			= false;
-	};
-
-#if PLATFORM_MAC_X86
-	/** Intel architecture descriptor of each GPU in the OS that provides stock details about the GPU that are innaccessible from the higher-level rendering APIs and provides a direct link to the GPU in the IORegistry. */
-	typedef class FGPUDescriptorX86_64 : public FGPUDescriptorCommon<FGPUDescriptorX86_64>
-	{
-		friend class FGPUDescriptorCommon<FGPUDescriptorX86_64>;
-
-		public:
-			FGPUDescriptorX86_64() = default;
-			FGPUDescriptorX86_64(FGPUDescriptorX86_64 const& Other);
-
-			virtual ~FGPUDescriptorX86_64();
-
-		protected:
-			void CopyFromImpl(FGPUDescriptorCommon const& Other);
-			TMap<FString, float> GetPerformanceStatisticsImpl() const;
-
-		public:
-			uint64 RegistryID;
-			uint32 PCIDevice; // This is really an io_registry_entry_t which is a mach port name
-
-	} FGPUDescriptor;
-
-#elif PLATFORM_MAC_ARM64
-	typedef class FGPUDescriptorARM64 : public FGPUDescriptorCommon<FGPUDescriptorARM64>
-	{
-		friend class FGPUDescriptorCommon<FGPUDescriptorARM64>;
-
-		public:
-			FGPUDescriptorARM64() = default;
-			FGPUDescriptorARM64(FGPUDescriptorARM64 const& Other);
-
-			virtual ~FGPUDescriptorARM64();
-
-		protected:
-			void CopyFromImpl(FGPUDescriptorCommon const& Other);
-			TMap<FString, float> GetPerformanceStatisticsImpl() const;
-
-		public:
-			uint64 RegistryID = 0;
-
-	} FGPUDescriptor;
-
-#else
-	#error "Undefined Mac platform"
-#endif
-
 	enum class EMacGPUNotification : uint8
 	{
 		Added,
@@ -218,6 +144,8 @@ struct CORE_API FMacPlatformMisc : public FApplePlatformMisc
 	/** Handle GPU change notifications. */
 	static void GPUChangeNotification(uint64_t DeviceRegistryID, EMacGPUNotification Notification);
 	
+	using FGPUDescriptor = FMacGPUDescriptor;
+
 	/** Returns the static list of GPUs in the current machine. */
 	static TArray<FGPUDescriptor> const& GetGPUDescriptors();
 	
@@ -232,8 +160,6 @@ struct CORE_API FMacPlatformMisc : public FApplePlatformMisc
 	/** Updates variables in GMacAppInfo that cannot be initialized before PlatformPostInit() */
 	static void PostInitMacAppInfoUpdate();
 
-	static CGDisplayModeRef GetSupportedDisplayMode(CGDirectDisplayID DisplayID, uint32 Width, uint32 Height);
-
 	inline static void ChooseHDRDeviceAndColorGamut(uint32 DeviceId, uint32 DisplayNitLevel, EDisplayOutputFormat& OutputDevice, EDisplayColorGamut& ColorGamut)
 	{
 		// ScRGB, 1000 or 2000 nits, DCI-P3
diff --git a/Engine/Source/Runtime/Core/Public/Mac/MacPlatformMiscEx.h b/Engine/Source/Runtime/Core/Public/Mac/MacPlatformMiscEx.h
new file mode 100644
--- /dev/null
+++ b/Engine/Source/Runtime/Core/Public/Mac/MacPlatformMiscEx.h
@@ -0,0 +1,23 @@
+// Copyright Epic Games, Inc. All Rights Reserved.
+
+/*=============================================================================================
+	MacPlatformMiscEx.h: Mac platform additional misc functions that requires sdk depdendencies
+==============================================================================================*/
+
+#pragma once
+
+#include "Mac/MacSystemIncludes.h"
+
+class FMacPlatformMiscEx
+{
+public:
+	CORE_API static NSOperatingSystemVersion GetNSOperatingSystemVersion();
+	CORE_API static CGDisplayModeRef GetSupportedDisplayMode(CGDirectDisplayID DisplayID, uint32 Width, uint32 Height);
+
+	/**
+	 * Returns if A < B returns -1, if A > B returns 1, else A == B and returns 0.
+	 */
+	CORE_API static int32 VersionCompare(const NSOperatingSystemVersion& VersionA, const NSOperatingSystemVersion& VersionB);
+
+	CORE_API static id<NSObject> CommandletActivity;
+};
diff --git a/Engine/Source/Runtime/Launch/Private/Mac/LaunchMac.cpp b/Engine/Source/Runtime/Launch/Private/Mac/LaunchMac.cpp
--- a/Engine/Source/Runtime/Launch/Private/Mac/LaunchMac.cpp
+++ b/Engine/Source/Runtime/Launch/Private/Mac/LaunchMac.cpp
@@ -7,6 +7,7 @@
 #include "LaunchEngineLoop.h"
 #include "HAL/ExceptionHandling.h"
 #include "Mac/MacPlatformCrashContext.h"
+#include "Mac/MacPlatformMiscEx.h"
 #include "Apple/LaunchCppToSwift.h"
 
 #if WITH_ENGINE
@@ -302,7 +303,7 @@ void EngineFullCrashHandler(const FGenericCrashContext& GenericContext)
 	NSDictionary* InfoDictionary = [[NSBundle mainBundle] infoDictionary];
 	NSString* MinimumSystemVersionString = (NSString*)InfoDictionary[@"LSMinimumSystemVersion"];
 	NSOperatingSystemVersion MinimumSystemVersion = { 0 };
-	NSOperatingSystemVersion CurrentSystemVersion = FMacPlatformMisc::GetNSOperatingSystemVersion();
+	NSOperatingSystemVersion CurrentSystemVersion = FMacPlatformMiscEx::GetNSOperatingSystemVersion();
 	NSOperatingSystemVersion MinSupportedMacOSVersion = { 12, 0, 0 };
 	NSString* MinSupportedMacOSVersionString = @"12.0.0";
 
@@ -312,13 +313,13 @@ void EngineFullCrashHandler(const FGenericCrashContext& GenericContext)
 	MinimumSystemVersion.patchVersion = VersionComponents.count > 2 ? [[VersionComponents objectAtIndex:2] integerValue] : 0;
 
 	// Make sure that the min version in Info.plist is at least 10.14.6, as that's the absolute minimum
-	if (FMacPlatformMisc::VersionCompare(MinimumSystemVersion, MinSupportedMacOSVersion) < 0)
+	if (FMacPlatformMiscEx::VersionCompare(MinimumSystemVersion, MinSupportedMacOSVersion) < 0)
 	{
 		MinimumSystemVersion = MinSupportedMacOSVersion;
 		MinimumSystemVersionString = MinSupportedMacOSVersionString;
 	}
 
-	if (FMacPlatformMisc::VersionCompare(CurrentSystemVersion, MinimumSystemVersion) < 0)
+	if (FMacPlatformMiscEx::VersionCompare(CurrentSystemVersion, MinimumSystemVersion) < 0)
 	{
 		CFDictionaryRef SessionDictionary = CGSessionCopyCurrentDictionary();
 		const bool bIsWindowServerAvailable = SessionDictionary != nullptr;
@@ -373,7 +374,7 @@ void EngineFullCrashHandler(const FGenericCrashContext& GenericContext)
 	
 #if !IS_MONOLITHIC
 	// UE-172403: dlopen crash on Ventura [13.0~13.3)
-	if (FMacPlatformMisc::VersionCompare(CurrentSystemVersion, {13, 0, 0}) >= 0 && FMacPlatformMisc::VersionCompare(CurrentSystemVersion, {13, 3, 0}) < 0)
+	if (FMacPlatformMiscEx::VersionCompare(CurrentSystemVersion, {13, 0, 0}) >= 0 && FMacPlatformMiscEx::VersionCompare(CurrentSystemVersion, {13, 3, 0}) < 0)
 	{
 		CFDictionaryRef SessionDictionary = CGSessionCopyCurrentDictionary();
 		const bool bIsWindowServerAvailable = SessionDictionary != nullptr;
