Change generated string name of tests using TEST_CASE macro to include __FILE__ and __LINE__ so that names are unique across modules & translation units

#Synced-CL 50158515

diff --git a/Engine/Source/Runtime/Core/Public/Misc/LowLevelTestAdapter.h b/Engine/Source/Runtime/Core/Public/Misc/LowLevelTestAdapter.h
--- a/Engine/Source/Runtime/Core/Public/Misc/LowLevelTestAdapter.h
+++ b/Engine/Source/Runtime/Core/Public/Misc/LowLevelTestAdapter.h
@@ -104,7 +104,7 @@ FString CaptureExpressionsAndValues(const FString& InExpressions, ArgTypes&&...
 #define TEST_CASE_GENERATED_NAME_UNIQUE LLT_JOIN(FLLTAdaptedTest, __COUNTER__)
 #define LLT_STR(Macro) #Macro
 #define LLT_STR_EXPAND(Macro) LLT_STR(Macro)
-#define TEST_CASE_GENERATED_NAME_UNIQUE_STR LLT_STR_EXPAND(TEST_CASE_GENERATED_NAME_UNIQUE)
+#define TEST_CASE_GENERATED_NAME_UNIQUE_STR "LLTAdaptedTest_" LLT_STR_EXPAND(__FILE__) "_" LLT_STR_EXPAND(__LINE__)
 // Note: TEST_CASE uses unique names which only work when used inside unique namespace in the same compilation unit?// Use TEST_CASE_NAMED instead and provide an unique global instance name
 #define TEST_CASE(PrettyName, TFlagsOrTags) TEST_CASE_NAMED_STR(TEST_CASE_GENERATED_NAME_UNIQUE, TEST_CASE_GENERATED_NAME_UNIQUE_STR, PrettyName, TFlagsOrTags)
 #define TEST_CASE_NAMED(ClassName, PrettyName, TFlagsOrTags) TEST_CASE_NAMED_STR(ClassName, #ClassName, PrettyName, TFlagsOrTags)
