[Engine] Bump up the MaterialTranslationDDCVersion to force recooks required by earlier CLs 50157519 and 50164238

diff --git a/Engine/Source/Runtime/Core/Private/UObject/DevObjectVersion.cpp b/Engine/Source/Runtime/Core/Private/UObject/DevObjectVersion.cpp
--- a/Engine/Source/Runtime/Core/Private/UObject/DevObjectVersion.cpp
+++ b/Engine/Source/Runtime/Core/Private/UObject/DevObjectVersion.cpp
@@ -138,7 +138,7 @@ FDevSystemGuids::FDevSystemGuids()
 	, POSESEARCHDB_DERIVEDDATA_VER(0x389117E4, 0x807A4CC0, 0x9F37C2E6, 0xD808A78D)
 	, SkeletalMeshDerivedDataVersion(0x9B5F4544, 0x76D7481C, 0x9AD3F614, 0xA6C07904)
 	, STATICMESH_DERIVEDDATA_VER(0x2C6C400C, 0x2EDF47B3, 0x9BD36689, 0x7B77B208)
-	, MaterialTranslationDDCVersion(0x628D9AFC, 0xC10245C5, 0x86967652, 0xAD52E887)
+	, MaterialTranslationDDCVersion(0x086219B8, 0x6E3F435D, 0xA9D99249, 0xD681D5DD)
 {
 }
 
