UnrealBuildTool: Enable -Wdangling by default for BuildSettingsVersion.V7
Fixed some lifetime violations.
Fixed false FAnsiStringView -> FUtf8StringView dangling reference warning.

diff --git a/Engine/Plugins/Experimental/UAF/UAF/Source/UAFEditor/Private/Functions/Outliner/FunctionsOutlinerHierarchy.cpp b/Engine/Plugins/Experimental/UAF/UAF/Source/UAFEditor/Private/Functions/Outliner/FunctionsOutlinerHierarchy.cpp
--- a/Engine/Plugins/Experimental/UAF/UAF/Source/UAFEditor/Private/Functions/Outliner/FunctionsOutlinerHierarchy.cpp
+++ b/Engine/Plugins/Experimental/UAF/UAF/Source/UAFEditor/Private/Functions/Outliner/FunctionsOutlinerHierarchy.cpp
@@ -80,8 +80,8 @@ FSceneOutlinerTreeItemPtr FFunctionsOutlinerHierarchy::FindOrCreateParentItem(co
 		int32 Index = INDEX_NONE;
 		if (CategoryPath.FindLastChar('|', Index))
 		{
-			ParentCategoryName = CategoryPath.Left(Index);
-			SubCategoryName = CategoryPath.RightChop(Index+1);
+			ParentCategoryName = MakeStringView(CategoryPath).Left(Index);
+			SubCategoryName = MakeStringView(CategoryPath).RightChop(Index+1);
 		}
 		else
 		{
diff --git a/Engine/Source/Programs/UnrealBuildTool/Configuration/Rules/CppCompileWarnings.cs b/Engine/Source/Programs/UnrealBuildTool/Configuration/Rules/CppCompileWarnings.cs
--- a/Engine/Source/Programs/UnrealBuildTool/Configuration/Rules/CppCompileWarnings.cs
+++ b/Engine/Source/Programs/UnrealBuildTool/Configuration/Rules/CppCompileWarnings.cs
@@ -1047,7 +1047,8 @@ namespace UnrealBuildTool
 		///		https://clang.llvm.org/docs/DiagnosticsReference.html#wdangling
 		/// </summary>	
 		[WarningsVCClang(["-Wno-dangling"], ["-Wdangling", "-Wno-error=dangling"], ["-Wdangling"])]
-		[BasicWarningLevelDefault(WarningLevel.Off)]
+		[VersionWarningLevelDefault(WarningLevel.Off, BuildSettingsVersion.V1, BuildSettingsVersion.V6, InitializationContext.Any)]
+		[VersionWarningLevelDefault(WarningLevel.Error, BuildSettingsVersion.V7, BuildSettingsVersion.Latest, InitializationContext.Any)]
 		public WarningLevel DanglingWarningLevel
 		{
 			get => ResolveWarning(ParentCppCompileWarnings?.DanglingWarningLevel, _danglingWarningLevel);
diff --git a/Engine/Source/Programs/UnrealBuildTool/Configuration/Rules/TargetRules.cs b/Engine/Source/Programs/UnrealBuildTool/Configuration/Rules/TargetRules.cs
--- a/Engine/Source/Programs/UnrealBuildTool/Configuration/Rules/TargetRules.cs
+++ b/Engine/Source/Programs/UnrealBuildTool/Configuration/Rules/TargetRules.cs
@@ -164,6 +164,7 @@ namespace UnrealBuildTool
 		/// <summary>
 		/// New defaults for 5.8:
 		/// * ModuleRules.CppCompileWarningSettings.ReturnTypeWarningLevel = WarningLevel.Error
+		/// * ModuleRules.CppCompileWarningSettings.DanglingWarningLevel = WarningLevel.Error
 		/// </summary>
 		V7,
 
@@ -3696,13 +3697,14 @@ namespace UnrealBuildTool
 				if (BuildSettingsVersion.V6 <= latestVersion && DefaultBuildSettings < BuildSettingsVersion.V6)
 				{
 					modifiedSettings.Add(Tuple.Create($"{nameof(WindowsPlatform)}.{nameof(WindowsPlatform.bStrictInlineConformance)} = true", "Updates MSVC inline conformance (Remove unreferenced COMDAT) to true (Previously: false)."));
-					modifiedSettings.Add(Tuple.Create(String.Format("{0} = WarningLevel.Error", nameof(CppCompileWarningSettings.UndefinedIdentifierWarningLevel)), "Enables compile-time validation of undefined identifiers in conditional expressions. (Previously: Off)."));
+					modifiedSettings.Add(Tuple.Create($"{nameof(CppCompileWarningSettings.UndefinedIdentifierWarningLevel)} = WarningLevel.Error", "Enables compile-time validation of undefined identifiers in conditional expressions. (Previously: Off)."));
 				}
 
 				if (BuildSettingsVersion.V7 <= latestVersion && DefaultBuildSettings < BuildSettingsVersion.V7)
 				{
 					modifiedSettings.Add(Tuple.Create($"{nameof(WindowsPlatform)}.{nameof(WindowsPlatform.bStrictEnumTypesConformance)} = true", "Updates MSVC enum types conformance to true (Previously: false)."));
-					modifiedSettings.Add(Tuple.Create(String.Format("{0} = WarningLevel.Error", nameof(CppCompileWarningSettings.ReturnTypeWarningLevel)), "Validates that non-void functions return a value. (Previously: Off for VCClang, Error for MSVC and non-VC Clang)."));
+					modifiedSettings.Add(Tuple.Create($"{nameof(CppCompileWarningSettings.ReturnTypeWarningLevel)} = WarningLevel.Error", "Validates that non-void functions return a value. (Previously: Off for VCClang, Error for MSVC and non-VC Clang)."));
+					modifiedSettings.Add(Tuple.Create($"{nameof(CppCompileWarningSettings.DanglingWarningLevel)} = WarningLevel.Error", "Enables clang warnings related possible dangling references or pointers. (Previously: Off for VCClang, Error for non-VC Clang)."));
 				}
 
 				if (modifiedSettings.Count > 0)
diff --git a/Engine/Source/Programs/UnrealBuildTool/Platform/Linux/LinuxToolChain.cs b/Engine/Source/Programs/UnrealBuildTool/Platform/Linux/LinuxToolChain.cs
--- a/Engine/Source/Programs/UnrealBuildTool/Platform/Linux/LinuxToolChain.cs
+++ b/Engine/Source/Programs/UnrealBuildTool/Platform/Linux/LinuxToolChain.cs
@@ -277,9 +277,6 @@ namespace UnrealBuildTool
 		{
 			base.GetCompileArguments_WarningsAndErrors(CompileEnvironment, Arguments);
 
-			//Arguments.Add("-Wunreachable-code");            // additional warning not normally included in Wall: warns if there is code that will never be executed - not helpful due to bIsGCC and similar
-			Arguments.Add("-Wno-dangling");		// kinda scary, but matches what we currently do on MS platforms when clang is used
-
 			Arguments.Add("-Wno-undefined-bool-conversion"); // hides checking if 'this' pointer is null
 		}
 
diff --git a/Engine/Source/Runtime/Core/Public/Containers/StringView.h b/Engine/Source/Runtime/Core/Public/Containers/StringView.h
--- a/Engine/Source/Runtime/Core/Public/Containers/StringView.h
+++ b/Engine/Source/Runtime/Core/Public/Containers/StringView.h
@@ -49,6 +49,21 @@ namespace UE::Core::Private
 	}
 } // UE::Core::Private
 
+template <typename T> inline constexpr bool TIsTStringView_V                 = false;
+template <typename T> inline constexpr bool TIsTStringView_V<TStringView<T>> = true;
+
+template <typename T> inline constexpr bool TIsTStringView_V<const          T> = TIsTStringView_V<T>;
+template <typename T> inline constexpr bool TIsTStringView_V<      volatile T> = TIsTStringView_V<T>;
+template <typename T> inline constexpr bool TIsTStringView_V<const volatile T> = TIsTStringView_V<T>;
+
+template <typename T>
+struct TIsTStringView
+{
+	static constexpr bool Value = TIsTStringView_V<T>;
+	static constexpr bool value = Value;
+};
+
+
 /**
  * A string view is a non-owning view of a range of characters.
  *
@@ -138,14 +153,24 @@ public:
 
 	/** Construct a view from a contiguous range of characters, such as FString or TStringBuilder. */
 	template <UE::CNotCVRefTo<ViewType> CharRangeType>
-		// !std::is_pointer is used here to force the pointer constructor to be used when passed an array or pointer.
-		requires (!std::is_pointer_v<std::decay_t<CharRangeType>> && UE::CCompatibleStringViewable<CharRangeType, CharType>)
+		// !std::is_pointer_v is used here to force the pointer constructor to be used when passed an array or pointer.
+		requires (!std::is_pointer_v<std::decay_t<CharRangeType>> && UE::CCompatibleStringViewable<CharRangeType, CharType> && !TIsTStringView_V<std::decay_t<CharRangeType>>)
 	[[nodiscard]] constexpr inline TStringView(const CharRangeType& InRange UE_LIFETIMEBOUND)
 		: DataPtr((const CharType*)GetDataHelper(InRange))
 		, Size(IntCastChecked<int32>(GetNum(InRange)))
 	{
 	}
 
+	/** Construct a view from a contiguous range of characters, such as FString or TStringBuilder. */
+	template <UE::CNotCVRefTo<ViewType> CharRangeType>
+	// !std::is_pointer_v is used here to force the pointer constructor to be used when passed an array or pointer.
+		requires (!std::is_pointer_v<std::decay_t<CharRangeType>> && UE::CCompatibleStringViewable<CharRangeType, CharType> && TIsTStringView_V<std::decay_t<CharRangeType>>)
+	[[nodiscard]] constexpr inline TStringView(const CharRangeType& InRange)
+		: DataPtr((const CharType*)GetDataHelper(InRange))
+		, Size(IntCastChecked<int32>(GetNum(InRange)))
+	{
+	}
+
 	//////////////////////////////////
 	// Start - assignment operators //
 	//////////////////////////////////
@@ -170,13 +195,22 @@ public:
 
 	/** Construct a view from a contiguous range of characters, such as FString or TStringBuilder. */
 	template <UE::CNotCVRefTo<ViewType> CharRangeType>
-		// !std::is_pointer is used here to force the pointer constructor to be used when passed an array or pointer.
-		requires (!std::is_pointer_v<std::decay_t<CharRangeType>> && UE::CCompatibleStringViewable<CharRangeType, CharType>)
+		// !std::is_pointer_v is used here to force the pointer constructor to be used when passed an array or pointer.
+		requires (!std::is_pointer_v<std::decay_t<CharRangeType>> && UE::CCompatibleStringViewable<CharRangeType, CharType> && !TIsTStringView_V<std::decay_t<CharRangeType>>)
 	UE_REWRITE constexpr TStringView& operator=(const CharRangeType& InRange UE_LIFETIMEBOUND)
 	{
 		return *this = TStringView(InRange);
 	}
 
+	/** Construct a view from a contiguous range of characters, such as FString or TStringBuilder. */
+	template <UE::CNotCVRefTo<ViewType> CharRangeType>
+	// !std::is_pointer_v is used here to force the pointer constructor to be used when passed an array or pointer.
+		requires (!std::is_pointer_v<std::decay_t<CharRangeType>> && UE::CCompatibleStringViewable<CharRangeType, CharType> && TIsTStringView_V<std::decay_t<CharRangeType>>)
+	UE_REWRITE constexpr TStringView& operator=(const CharRangeType& InRange)
+	{
+		return *this = TStringView(InRange);
+	}
+
 	////////////////////////////////
 	// End - assignment operators //
 	////////////////////////////////
