Delegate name fixes.

diff --git a/Engine/Source/Runtime/Core/Public/Delegates/Delegate.h b/Engine/Source/Runtime/Core/Public/Delegates/Delegate.h
--- a/Engine/Source/Runtime/Core/Public/Delegates/Delegate.h
+++ b/Engine/Source/Runtime/Core/Public/Delegates/Delegate.h
@@ -202,47 +202,36 @@
  * @note: The last parameter is variadic and is used as the 'template args' for this delegate's classes (__VA_ARGS__)
  * @note: To avoid issues with macro expansion breaking code navigation, make sure the type/class name macro params are unique across all of these macros
  */
-#define UE_PRIVATE_DECLARE_DELEGATE(DelegateName, ReturnType, ...) \
-	typedef TDelegate<ReturnType(__VA_ARGS__)> DelegateName;
+#define UE_PRIVATE_DECLARE_DELEGATE(DelegateName, FuncType) \
+	typedef TDelegate<FuncType> DelegateName;
 
 /** Declares a broadcast delegate that can bind to multiple native functions simultaneously */
-#define UE_PRIVATE_DECLARE_MULTICAST_DELEGATE(MulticastDelegateName, ReturnType, ...) \
-	typedef TMulticastDelegate<ReturnType(__VA_ARGS__)> MulticastDelegateName;
+#define UE_PRIVATE_DECLARE_MULTICAST_DELEGATE(MulticastDelegateName, FuncType) \
+	typedef TMulticastDelegate<FuncType> MulticastDelegateName;
 
 /** Declares a broadcast thread-safe delegate that can bind to multiple native functions simultaneously */
-#define UE_PRIVATE_DECLARE_TS_MULTICAST_DELEGATE(MulticastDelegateName, ReturnType, ...) \
-	typedef TMulticastDelegate<ReturnType(__VA_ARGS__), FDefaultTSDelegateUserPolicy> MulticastDelegateName;
+#define UE_PRIVATE_DECLARE_TS_MULTICAST_DELEGATE(MulticastDelegateName, FuncType) \
+	typedef TMulticastDelegate<FuncType, FDefaultTSDelegateUserPolicy> MulticastDelegateName;
 
 /**
  * Declares a multicast delegate that is meant to only be activated from OwningType
  *
  * @note: This behavior is not enforced and this type should be considered deprecated for new delegates, use normal multicast instead
  */
-#define UE_PRIVATE_DECLARE_EVENT(OwningType, EventName, ReturnType, ...) \
-	class EventName : public TMulticastDelegate<ReturnType(__VA_ARGS__)> \
-	{ \
-		friend class OwningType; \
-	};
+#define UE_PRIVATE_DECLARE_EVENT(OwningType, EventName, FuncType) \
+	using EventName = TMulticastDelegate<FuncType>;
 
 /** Declares a derived event delegate that works the same as its parent type but is intended to be used by a different owning type */
 #define DECLARE_DERIVED_EVENT(OwningType, BaseTypeEvent, EventName) \
 	class EventName : public BaseTypeEvent { friend class OwningType; };
 
 /** Declare user's dynamic delegate, with wrapper proxy method for executing the delegate */
-#define UE_PRIVATE_DECLARE_DYNAMIC_DELEGATE(DynamicDelegateClassName, ...) \
-	class DynamicDelegateClassName : public TBaseDynamicDelegate<FNotThreadSafeDelegateMode, __VA_ARGS__> \
-	{ \
-	public: \
-		using TBaseDynamicDelegate<FNotThreadSafeDelegateMode, __VA_ARGS__>::TBaseDynamicDelegate; \
-	};
+#define UE_PRIVATE_DECLARE_DYNAMIC_DELEGATE(DynamicDelegateClassName, FuncType) \
+	using DynamicDelegateClassName = TDynamicDelegate<FuncType>;
 
 /** Declare user's dynamic multi-cast delegate, with wrapper proxy method for executing the delegate */
-#define UE_PRIVATE_DECLARE_DYNAMIC_MULTICAST_DELEGATE(DynamicMulticastDelegateClassName, ...) \
-class DynamicMulticastDelegateClassName : public TBaseDynamicMulticastDelegate<FNotThreadSafeDelegateMode, __VA_ARGS__> \
-	{ \
-	public: \
-		using TBaseDynamicMulticastDelegate<FNotThreadSafeDelegateMode, __VA_ARGS__>::TBaseDynamicMulticastDelegate; \
-	};
+#define UE_PRIVATE_DECLARE_DYNAMIC_MULTICAST_DELEGATE(DynamicMulticastDelegateClassName, FuncType) \
+	using DynamicMulticastDelegateClassName = TDynamicMulticastDelegate<FuncType>;
 
 // These macros were an implementation detail and should not be used by user code.  The appropriate DECLARE_* macro in DelegateCombinations.h should be used instead.
 #define FUNC_DECLARE_DELEGATE(DelegateName, ReturnType, ...)                            UE_DEPRECATED_MACRO(5.8, "FUNC_DECLARE_DELEGATE has been deprecated.")                   UE_PRIVATE_DECLARE_DELEGATE(DelegateName, ReturnType, ##__VA_ARGS__)
diff --git a/Engine/Source/Runtime/Core/Public/Delegates/DelegateCombinations.h b/Engine/Source/Runtime/Core/Public/Delegates/DelegateCombinations.h
--- a/Engine/Source/Runtime/Core/Public/Delegates/DelegateCombinations.h
+++ b/Engine/Source/Runtime/Core/Public/Delegates/DelegateCombinations.h
@@ -17,111 +17,111 @@
 #endif
 
 /** Declares a delegate that can only bind to one native function at a time */
-#define DECLARE_DELEGATE(DelegateName) UE_PRIVATE_DECLARE_DELEGATE(DelegateName, void)
+#define DECLARE_DELEGATE(DelegateName) UE_PRIVATE_DECLARE_DELEGATE(DelegateName, void())
 
 /** Declares a broadcast delegate that can bind to multiple native functions simultaneously */
-#define DECLARE_MULTICAST_DELEGATE(DelegateName) UE_PRIVATE_DECLARE_MULTICAST_DELEGATE(DelegateName, void)
+#define DECLARE_MULTICAST_DELEGATE(DelegateName) UE_PRIVATE_DECLARE_MULTICAST_DELEGATE(DelegateName, void())
 
 /** Declares a broadcast thread-safe delegate that can bind to multiple native functions simultaneously */
-#define DECLARE_TS_MULTICAST_DELEGATE(DelegateName) UE_PRIVATE_DECLARE_TS_MULTICAST_DELEGATE(DelegateName, void)
+#define DECLARE_TS_MULTICAST_DELEGATE(DelegateName) UE_PRIVATE_DECLARE_TS_MULTICAST_DELEGATE(DelegateName, void())
 
 /**
  * Declares a multicast delegate that is meant to only be activated from OwningType
  * NOTE: This behavior is not enforced and this type should be considered deprecated for new delegates, use normal multicast instead
  */
-#define DECLARE_EVENT(OwningType, EventName) UE_PRIVATE_DECLARE_EVENT(OwningType, EventName, void)
+#define DECLARE_EVENT(OwningType, EventName) UE_PRIVATE_DECLARE_EVENT(OwningType, EventName, void())
 
 /** Declares a blueprint-accessible delegate that can only bind to one UFUNCTION at a time */
-#define DECLARE_DYNAMIC_DELEGATE(DelegateName) UE_PRIVATE_DECLARE_DYNAMIC_DELEGATE(DelegateName, void)
+#define DECLARE_DYNAMIC_DELEGATE(DelegateName) UE_PRIVATE_DECLARE_DYNAMIC_DELEGATE(DelegateName, void())
 
 /** Declares a blueprint-accessible broadcast delegate that can bind to multiple native UFUNCTIONs simultaneously */
-#define DECLARE_DYNAMIC_MULTICAST_DELEGATE(DelegateName) UE_PRIVATE_DECLARE_DYNAMIC_MULTICAST_DELEGATE(DelegateName, void)
+#define DECLARE_DYNAMIC_MULTICAST_DELEGATE(DelegateName) UE_PRIVATE_DECLARE_DYNAMIC_MULTICAST_DELEGATE(DelegateName, void())
 
 /** Declares a delegate with return value that can only bind to one native function at a time */
-#define DECLARE_DELEGATE_RetVal(ReturnValueType, DelegateName) UE_PRIVATE_DECLARE_DELEGATE(DelegateName, ReturnValueType)
+#define DECLARE_DELEGATE_RetVal(ReturnValueType, DelegateName) UE_PRIVATE_DECLARE_DELEGATE(DelegateName, ReturnValueType())
 
-/** Declares a blueprint-accessible delegate with return value that can only bind to one UFUNCTION at a time */
-#define DECLARE_DYNAMIC_DELEGATE_RetVal(ReturnValueType, DelegateName) UE_PRIVATE_DECLARE_DYNAMIC_DELEGATE(DelegateName, ReturnValueType)
+/** Declares a blueprint-accessible delegate with return value that can only bind to one UNFUNCTION at a time */
+#define DECLARE_DYNAMIC_DELEGATE_RetVal(ReturnValueType, DelegateName) UE_PRIVATE_DECLARE_DYNAMIC_DELEGATE(DelegateName, ReturnValueType())
 
 
 // Multiple-parameter versions of above delegate types:
-#define DECLARE_DELEGATE_OneParam(DelegateName, Param1Type) UE_PRIVATE_DECLARE_DELEGATE(DelegateName, void, Param1Type)
-#define DECLARE_MULTICAST_DELEGATE_OneParam(DelegateName, Param1Type) UE_PRIVATE_DECLARE_MULTICAST_DELEGATE(DelegateName, void, Param1Type)
-#define DECLARE_TS_MULTICAST_DELEGATE_OneParam(DelegateName, Param1Type) UE_PRIVATE_DECLARE_TS_MULTICAST_DELEGATE(DelegateName, void, Param1Type)
-#define DECLARE_EVENT_OneParam(OwningType, EventName, Param1Type) UE_PRIVATE_DECLARE_EVENT(OwningType, EventName, void, Param1Type)
-#define DECLARE_DYNAMIC_DELEGATE_OneParam(DelegateName, Param1Type, Param1Name) UE_PRIVATE_DECLARE_DYNAMIC_DELEGATE(DelegateName, void, Param1Type)
-#define DECLARE_DYNAMIC_MULTICAST_DELEGATE_OneParam(DelegateName, Param1Type, Param1Name) UE_PRIVATE_DECLARE_DYNAMIC_MULTICAST_DELEGATE(DelegateName, void, Param1Type)
-#define DECLARE_DELEGATE_RetVal_OneParam(ReturnValueType, DelegateName, Param1Type) UE_PRIVATE_DECLARE_DELEGATE(DelegateName, ReturnValueType, Param1Type)
-#define DECLARE_DYNAMIC_DELEGATE_RetVal_OneParam(ReturnValueType, DelegateName, Param1Type, Param1Name) UE_PRIVATE_DECLARE_DYNAMIC_DELEGATE(DelegateName, ReturnValueType, Param1Type)
-
-#define DECLARE_DELEGATE_TwoParams(DelegateName, Param1Type, Param2Type) UE_PRIVATE_DECLARE_DELEGATE(DelegateName, void, Param1Type, Param2Type)
-#define DECLARE_MULTICAST_DELEGATE_TwoParams(DelegateName, Param1Type, Param2Type) UE_PRIVATE_DECLARE_MULTICAST_DELEGATE(DelegateName, void, Param1Type, Param2Type)
-#define DECLARE_TS_MULTICAST_DELEGATE_TwoParams(DelegateName, Param1Type, Param2Type) UE_PRIVATE_DECLARE_TS_MULTICAST_DELEGATE(DelegateName, void, Param1Type, Param2Type)
-#define DECLARE_EVENT_TwoParams(OwningType, EventName, Param1Type, Param2Type) UE_PRIVATE_DECLARE_EVENT(OwningType, EventName, void, Param1Type, Param2Type)
-#define DECLARE_DYNAMIC_DELEGATE_TwoParams(DelegateName, Param1Type, Param1Name, Param2Type, Param2Name) UE_PRIVATE_DECLARE_DYNAMIC_DELEGATE(DelegateName, void, Param1Type, Param2Type)
-#define DECLARE_DYNAMIC_MULTICAST_DELEGATE_TwoParams(DelegateName, Param1Type, Param1Name, Param2Type, Param2Name) UE_PRIVATE_DECLARE_DYNAMIC_MULTICAST_DELEGATE(DelegateName, void, Param1Type, Param2Type)
-#define DECLARE_DELEGATE_RetVal_TwoParams(ReturnValueType, DelegateName, Param1Type, Param2Type) UE_PRIVATE_DECLARE_DELEGATE(DelegateName, ReturnValueType, Param1Type, Param2Type)
-#define DECLARE_DYNAMIC_DELEGATE_RetVal_TwoParams(ReturnValueType, DelegateName, Param1Type, Param1Name, Param2Type, Param2Name) UE_PRIVATE_DECLARE_DYNAMIC_DELEGATE(DelegateName, ReturnValueType, Param1Type, Param2Type)
-
-#define DECLARE_DELEGATE_ThreeParams(DelegateName, Param1Type, Param2Type, Param3Type) UE_PRIVATE_DECLARE_DELEGATE(DelegateName, void, Param1Type, Param2Type, Param3Type)
-#define DECLARE_MULTICAST_DELEGATE_ThreeParams(DelegateName, Param1Type, Param2Type, Param3Type) UE_PRIVATE_DECLARE_MULTICAST_DELEGATE(DelegateName, void, Param1Type, Param2Type, Param3Type)
-#define DECLARE_TS_MULTICAST_DELEGATE_ThreeParams(DelegateName, Param1Type, Param2Type, Param3Type) UE_PRIVATE_DECLARE_TS_MULTICAST_DELEGATE(DelegateName, void, Param1Type, Param2Type, Param3Type)
-#define DECLARE_EVENT_ThreeParams(OwningType, EventName, Param1Type, Param2Type, Param3Type) UE_PRIVATE_DECLARE_EVENT(OwningType, EventName, void, Param1Type, Param2Type, Param3Type)
-#define DECLARE_DYNAMIC_DELEGATE_ThreeParams(DelegateName, Param1Type, Param1Name, Param2Type, Param2Name, Param3Type, Param3Name) UE_PRIVATE_DECLARE_DYNAMIC_DELEGATE(DelegateName, void, Param1Type, Param2Type, Param3Type)
-#define DECLARE_DYNAMIC_MULTICAST_DELEGATE_ThreeParams(DelegateName, Param1Type, Param1Name, Param2Type, Param2Name, Param3Type, Param3Name) UE_PRIVATE_DECLARE_DYNAMIC_MULTICAST_DELEGATE(DelegateName, void, Param1Type, Param2Type, Param3Type)
-#define DECLARE_DELEGATE_RetVal_ThreeParams(ReturnValueType, DelegateName, Param1Type, Param2Type, Param3Type) UE_PRIVATE_DECLARE_DELEGATE(DelegateName, ReturnValueType, Param1Type, Param2Type, Param3Type)
-#define DECLARE_DYNAMIC_DELEGATE_RetVal_ThreeParams(ReturnValueType, DelegateName, Param1Type, Param1Name, Param2Type, Param2Name, Param3Type, Param3Name) UE_PRIVATE_DECLARE_DYNAMIC_DELEGATE(DelegateName, ReturnValueType, Param1Type, Param2Type, Param3Type)
-
-#define DECLARE_DELEGATE_FourParams(DelegateName, Param1Type, Param2Type, Param3Type, Param4Type) UE_PRIVATE_DECLARE_DELEGATE(DelegateName, void, Param1Type, Param2Type, Param3Type, Param4Type)
-#define DECLARE_MULTICAST_DELEGATE_FourParams(DelegateName, Param1Type, Param2Type, Param3Type, Param4Type) UE_PRIVATE_DECLARE_MULTICAST_DELEGATE(DelegateName, void, Param1Type, Param2Type, Param3Type, Param4Type)
-#define DECLARE_TS_MULTICAST_DELEGATE_FourParams(DelegateName, Param1Type, Param2Type, Param3Type, Param4Type) UE_PRIVATE_DECLARE_TS_MULTICAST_DELEGATE(DelegateName, void, Param1Type, Param2Type, Param3Type, Param4Type)
-#define DECLARE_EVENT_FourParams(OwningType, EventName, Param1Type, Param2Type, Param3Type, Param4Type) UE_PRIVATE_DECLARE_EVENT(OwningType, EventName, void, Param1Type, Param2Type, Param3Type, Param4Type)
-#define DECLARE_DYNAMIC_DELEGATE_FourParams(DelegateName, Param1Type, Param1Name, Param2Type, Param2Name, Param3Type, Param3Name, Param4Type, Param4Name) UE_PRIVATE_DECLARE_DYNAMIC_DELEGATE(DelegateName, void, Param1Type, Param2Type, Param3Type, Param4Type)
-#define DECLARE_DYNAMIC_MULTICAST_DELEGATE_FourParams(DelegateName, Param1Type, Param1Name, Param2Type, Param2Name, Param3Type, Param3Name, Param4Type, Param4Name) UE_PRIVATE_DECLARE_DYNAMIC_MULTICAST_DELEGATE(DelegateName, void, Param1Type, Param2Type, Param3Type, Param4Type)
-#define DECLARE_DELEGATE_RetVal_FourParams(ReturnValueType, DelegateName, Param1Type, Param2Type, Param3Type, Param4Type) UE_PRIVATE_DECLARE_DELEGATE(DelegateName, ReturnValueType, Param1Type, Param2Type, Param3Type, Param4Type)
-#define DECLARE_DYNAMIC_DELEGATE_RetVal_FourParams(ReturnValueType, DelegateName, Param1Type, Param1Name, Param2Type, Param2Name, Param3Type, Param3Name, Param4Type, Param4Name) UE_PRIVATE_DECLARE_DYNAMIC_DELEGATE(DelegateName, ReturnValueType, Param1Type, Param2Type, Param3Type, Param4Type)
-
-#define DECLARE_DELEGATE_FiveParams(DelegateName, Param1Type, Param2Type, Param3Type, Param4Type, Param5Type) UE_PRIVATE_DECLARE_DELEGATE(DelegateName, void, Param1Type, Param2Type, Param3Type, Param4Type, Param5Type)
-#define DECLARE_MULTICAST_DELEGATE_FiveParams(DelegateName, Param1Type, Param2Type, Param3Type, Param4Type, Param5Type) UE_PRIVATE_DECLARE_MULTICAST_DELEGATE(DelegateName, void, Param1Type, Param2Type, Param3Type, Param4Type, Param5Type)
-#define DECLARE_TS_MULTICAST_DELEGATE_FiveParams(DelegateName, Param1Type, Param2Type, Param3Type, Param4Type, Param5Type) UE_PRIVATE_DECLARE_TS_MULTICAST_DELEGATE(DelegateName, void, Param1Type, Param2Type, Param3Type, Param4Type, Param5Type)
-#define DECLARE_EVENT_FiveParams(OwningType, EventName, Param1Type, Param2Type, Param3Type, Param4Type, Param5Type) UE_PRIVATE_DECLARE_EVENT(OwningType, EventName, void, Param1Type, Param2Type, Param3Type, Param4Type, Param5Type)
-#define DECLARE_DYNAMIC_DELEGATE_FiveParams(DelegateName, Param1Type, Param1Name, Param2Type, Param2Name, Param3Type, Param3Name, Param4Type, Param4Name, Param5Type, Param5Name) UE_PRIVATE_DECLARE_DYNAMIC_DELEGATE(DelegateName, void, Param1Type, Param2Type, Param3Type, Param4Type, Param5Type)
-#define DECLARE_DYNAMIC_MULTICAST_DELEGATE_FiveParams(DelegateName, Param1Type, Param1Name, Param2Type, Param2Name, Param3Type, Param3Name, Param4Type, Param4Name, Param5Type, Param5Name) UE_PRIVATE_DECLARE_DYNAMIC_MULTICAST_DELEGATE(DelegateName, void, Param1Type, Param2Type, Param3Type, Param4Type, Param5Type)
-#define DECLARE_DELEGATE_RetVal_FiveParams(ReturnValueType, DelegateName, Param1Type, Param2Type, Param3Type, Param4Type, Param5Type) UE_PRIVATE_DECLARE_DELEGATE(DelegateName, ReturnValueType, Param1Type, Param2Type, Param3Type, Param4Type, Param5Type)
-#define DECLARE_DYNAMIC_DELEGATE_RetVal_FiveParams(ReturnValueType, DelegateName, Param1Type, Param1Name, Param2Type, Param2Name, Param3Type, Param3Name, Param4Type, Param4Name, Param5Type, Param5Name) UE_PRIVATE_DECLARE_DYNAMIC_DELEGATE(DelegateName, ReturnValueType, Param1Type, Param2Type, Param3Type, Param4Type, Param5Type)
-
-#define DECLARE_DELEGATE_SixParams(DelegateName, Param1Type, Param2Type, Param3Type, Param4Type, Param5Type, Param6Type) UE_PRIVATE_DECLARE_DELEGATE(DelegateName, void, Param1Type, Param2Type, Param3Type, Param4Type, Param5Type, Param6Type)
-#define DECLARE_MULTICAST_DELEGATE_SixParams(DelegateName, Param1Type, Param2Type, Param3Type, Param4Type, Param5Type, Param6Type) UE_PRIVATE_DECLARE_MULTICAST_DELEGATE(DelegateName, void, Param1Type, Param2Type, Param3Type, Param4Type, Param5Type, Param6Type)
-#define DECLARE_TS_MULTICAST_DELEGATE_SixParams(DelegateName, Param1Type, Param2Type, Param3Type, Param4Type, Param5Type, Param6Type) UE_PRIVATE_DECLARE_TS_MULTICAST_DELEGATE(DelegateName, void, Param1Type, Param2Type, Param3Type, Param4Type, Param5Type, Param6Type)
-#define DECLARE_EVENT_SixParams(OwningType, EventName, Param1Type, Param2Type, Param3Type, Param4Type, Param5Type, Param6Type) UE_PRIVATE_DECLARE_EVENT(OwningType, EventName, void, Param1Type, Param2Type, Param3Type, Param4Type, Param5Type, Param6Type)
-#define DECLARE_DYNAMIC_DELEGATE_SixParams(DelegateName, Param1Type, Param1Name, Param2Type, Param2Name, Param3Type, Param3Name, Param4Type, Param4Name, Param5Type, Param5Name, Param6Type, Param6Name) UE_PRIVATE_DECLARE_DYNAMIC_DELEGATE(DelegateName, void, Param1Type, Param2Type, Param3Type, Param4Type, Param5Type, Param6Type)
-#define DECLARE_DYNAMIC_MULTICAST_DELEGATE_SixParams(DelegateName, Param1Type, Param1Name, Param2Type, Param2Name, Param3Type, Param3Name, Param4Type, Param4Name, Param5Type, Param5Name, Param6Type, Param6Name) UE_PRIVATE_DECLARE_DYNAMIC_MULTICAST_DELEGATE(DelegateName, void, Param1Type, Param2Type, Param3Type, Param4Type, Param5Type, Param6Type)
-#define DECLARE_DELEGATE_RetVal_SixParams(ReturnValueType, DelegateName, Param1Type, Param2Type, Param3Type, Param4Type, Param5Type, Param6Type) UE_PRIVATE_DECLARE_DELEGATE(DelegateName, ReturnValueType, Param1Type, Param2Type, Param3Type, Param4Type, Param5Type, Param6Type)
-#define DECLARE_DYNAMIC_DELEGATE_RetVal_SixParams(ReturnValueType, DelegateName, Param1Type, Param1Name, Param2Type, Param2Name, Param3Type, Param3Name, Param4Type, Param4Name, Param5Type, Param5Name, Param6Type, Param6Name) UE_PRIVATE_DECLARE_DYNAMIC_DELEGATE(DelegateName, ReturnValueType, Param1Type, Param2Type, Param3Type, Param4Type, Param5Type, Param6Type)
-
-#define DECLARE_DELEGATE_SevenParams(DelegateName, Param1Type, Param2Type, Param3Type, Param4Type, Param5Type, Param6Type, Param7Type) UE_PRIVATE_DECLARE_DELEGATE(DelegateName, void, Param1Type, Param2Type, Param3Type, Param4Type, Param5Type, Param6Type, Param7Type)
-#define DECLARE_MULTICAST_DELEGATE_SevenParams(DelegateName, Param1Type, Param2Type, Param3Type, Param4Type, Param5Type, Param6Type, Param7Type) UE_PRIVATE_DECLARE_MULTICAST_DELEGATE(DelegateName, void, Param1Type, Param2Type, Param3Type, Param4Type, Param5Type, Param6Type, Param7Type)
-#define DECLARE_TS_MULTICAST_DELEGATE_SevenParams(DelegateName, Param1Type, Param2Type, Param3Type, Param4Type, Param5Type, Param6Type, Param7Type) UE_PRIVATE_DECLARE_TS_MULTICAST_DELEGATE(DelegateName, void, Param1Type, Param2Type, Param3Type, Param4Type, Param5Type, Param6Type, Param7Type)
-#define DECLARE_EVENT_SevenParams(OwningType, EventName, Param1Type, Param2Type, Param3Type, Param4Type, Param5Type, Param6Type, Param7Type) UE_PRIVATE_DECLARE_EVENT(OwningType, EventName, void, Param1Type, Param2Type, Param3Type, Param4Type, Param5Type, Param6Type, Param7Type)
-#define DECLARE_DYNAMIC_DELEGATE_SevenParams(DelegateName, Param1Type, Param1Name, Param2Type, Param2Name, Param3Type, Param3Name, Param4Type, Param4Name, Param5Type, Param5Name, Param6Type, Param6Name, Param7Type, Param7Name) UE_PRIVATE_DECLARE_DYNAMIC_DELEGATE(DelegateName, void, Param1Type, Param2Type, Param3Type, Param4Type, Param5Type, Param6Type, Param7Type)
-#define DECLARE_DYNAMIC_MULTICAST_DELEGATE_SevenParams(DelegateName, Param1Type, Param1Name, Param2Type, Param2Name, Param3Type, Param3Name, Param4Type, Param4Name, Param5Type, Param5Name, Param6Type, Param6Name, Param7Type, Param7Name) UE_PRIVATE_DECLARE_DYNAMIC_MULTICAST_DELEGATE(DelegateName, void, Param1Type, Param2Type, Param3Type, Param4Type, Param5Type, Param6Type, Param7Type)
-#define DECLARE_DELEGATE_RetVal_SevenParams(ReturnValueType, DelegateName, Param1Type, Param2Type, Param3Type, Param4Type, Param5Type, Param6Type, Param7Type) UE_PRIVATE_DECLARE_DELEGATE(DelegateName, ReturnValueType, Param1Type, Param2Type, Param3Type, Param4Type, Param5Type, Param6Type, Param7Type)
-#define DECLARE_DYNAMIC_DELEGATE_RetVal_SevenParams(ReturnValueType, DelegateName, Param1Type, Param1Name, Param2Type, Param2Name, Param3Type, Param3Name, Param4Type, Param4Name, Param5Type, Param5Name, Param6Type, Param6Name, Param7Type, Param7Name) UE_PRIVATE_DECLARE_DYNAMIC_DELEGATE(DelegateName, ReturnValueType, Param1Type, Param2Type, Param3Type, Param4Type, Param5Type, Param6Type, Param7Type)
-
-#define DECLARE_DELEGATE_EightParams(DelegateName, Param1Type, Param2Type, Param3Type, Param4Type, Param5Type, Param6Type, Param7Type, Param8Type) UE_PRIVATE_DECLARE_DELEGATE(DelegateName, void, Param1Type, Param2Type, Param3Type, Param4Type, Param5Type, Param6Type, Param7Type, Param8Type)
-#define DECLARE_MULTICAST_DELEGATE_EightParams(DelegateName, Param1Type, Param2Type, Param3Type, Param4Type, Param5Type, Param6Type, Param7Type, Param8Type) UE_PRIVATE_DECLARE_MULTICAST_DELEGATE(DelegateName, void, Param1Type, Param2Type, Param3Type, Param4Type, Param5Type, Param6Type, Param7Type, Param8Type)
-#define DECLARE_TS_MULTICAST_DELEGATE_EightParams(DelegateName, Param1Type, Param2Type, Param3Type, Param4Type, Param5Type, Param6Type, Param7Type, Param8Type) UE_PRIVATE_DECLARE_TS_MULTICAST_DELEGATE(DelegateName, void, Param1Type, Param2Type, Param3Type, Param4Type, Param5Type, Param6Type, Param7Type, Param8Type)
-#define DECLARE_EVENT_EightParams(OwningType, EventName, Param1Type, Param2Type, Param3Type, Param4Type, Param5Type, Param6Type, Param7Type, Param8Type) UE_PRIVATE_DECLARE_EVENT(OwningType, EventName, void, Param1Type, Param2Type, Param3Type, Param4Type, Param5Type, Param6Type, Param7Type, Param8Type)
-#define DECLARE_DYNAMIC_DELEGATE_EightParams(DelegateName, Param1Type, Param1Name, Param2Type, Param2Name, Param3Type, Param3Name, Param4Type, Param4Name, Param5Type, Param5Name, Param6Type, Param6Name, Param7Type, Param7Name, Param8Type, Param8Name) UE_PRIVATE_DECLARE_DYNAMIC_DELEGATE(DelegateName, void, Param1Type, Param2Type, Param3Type, Param4Type, Param5Type, Param6Type, Param7Type, Param8Type)
-#define DECLARE_DYNAMIC_MULTICAST_DELEGATE_EightParams(DelegateName, Param1Type, Param1Name, Param2Type, Param2Name, Param3Type, Param3Name, Param4Type, Param4Name, Param5Type, Param5Name, Param6Type, Param6Name, Param7Type, Param7Name, Param8Type, Param8Name) UE_PRIVATE_DECLARE_DYNAMIC_MULTICAST_DELEGATE(DelegateName, void, Param1Type, Param2Type, Param3Type, Param4Type, Param5Type, Param6Type, Param7Type, Param8Type)
-#define DECLARE_DELEGATE_RetVal_EightParams(ReturnValueType, DelegateName, Param1Type, Param2Type, Param3Type, Param4Type, Param5Type, Param6Type, Param7Type, Param8Type) UE_PRIVATE_DECLARE_DELEGATE(DelegateName, ReturnValueType, Param1Type, Param2Type, Param3Type, Param4Type, Param5Type, Param6Type, Param7Type, Param8Type)
-#define DECLARE_DYNAMIC_DELEGATE_RetVal_EightParams(ReturnValueType, DelegateName, Param1Type, Param1Name, Param2Type, Param2Name, Param3Type, Param3Name, Param4Type, Param4Name, Param5Type, Param5Name, Param6Type, Param6Name, Param7Type, Param7Name, Param8Type, Param8Name) UE_PRIVATE_DECLARE_DYNAMIC_DELEGATE(DelegateName, ReturnValueType, Param1Type, Param2Type, Param3Type, Param4Type, Param5Type, Param6Type, Param7Type, Param8Type)
-
-#define DECLARE_DELEGATE_NineParams(DelegateName, Param1Type, Param2Type, Param3Type, Param4Type, Param5Type, Param6Type, Param7Type, Param8Type, Param9Type) UE_PRIVATE_DECLARE_DELEGATE(DelegateName, void, Param1Type, Param2Type, Param3Type, Param4Type, Param5Type, Param6Type, Param7Type, Param8Type, Param9Type)
-#define DECLARE_MULTICAST_DELEGATE_NineParams(DelegateName, Param1Type, Param2Type, Param3Type, Param4Type, Param5Type, Param6Type, Param7Type, Param8Type, Param9Type) UE_PRIVATE_DECLARE_MULTICAST_DELEGATE(DelegateName, void, Param1Type, Param2Type, Param3Type, Param4Type, Param5Type, Param6Type, Param7Type, Param8Type, Param9Type)
-#define DECLARE_TS_MULTICAST_DELEGATE_NineParams(DelegateName, Param1Type, Param2Type, Param3Type, Param4Type, Param5Type, Param6Type, Param7Type, Param8Type, Param9Type) UE_PRIVATE_DECLARE_TS_MULTICAST_DELEGATE(DelegateName, void, Param1Type, Param2Type, Param3Type, Param4Type, Param5Type, Param6Type, Param7Type, Param8Type, Param9Type)
-#define DECLARE_EVENT_NineParams(OwningType, EventName, Param1Type, Param2Type, Param3Type, Param4Type, Param5Type, Param6Type, Param7Type, Param8Type, Param9Type) UE_PRIVATE_DECLARE_EVENT(OwningType, EventName, void, Param1Type, Param2Type, Param3Type, Param4Type, Param5Type, Param6Type, Param7Type, Param8Type, Param9Type)
-#define DECLARE_DYNAMIC_DELEGATE_NineParams(DelegateName, Param1Type, Param1Name, Param2Type, Param2Name, Param3Type, Param3Name, Param4Type, Param4Name, Param5Type, Param5Name, Param6Type, Param6Name, Param7Type, Param7Name, Param8Type, Param8Name, Param9Type, Param9Name) UE_PRIVATE_DECLARE_DYNAMIC_DELEGATE(DelegateName, void, Param1Type, Param2Type, Param3Type, Param4Type, Param5Type, Param6Type, Param7Type, Param8Type, Param9Type)
-#define DECLARE_DYNAMIC_MULTICAST_DELEGATE_NineParams(DelegateName, Param1Type, Param1Name, Param2Type, Param2Name, Param3Type, Param3Name, Param4Type, Param4Name, Param5Type, Param5Name, Param6Type, Param6Name, Param7Type, Param7Name, Param8Type, Param8Name, Param9Type, Param9Name) UE_PRIVATE_DECLARE_DYNAMIC_MULTICAST_DELEGATE(DelegateName, void, Param1Type, Param2Type, Param3Type, Param4Type, Param5Type, Param6Type, Param7Type, Param8Type, Param9Type)
-#define DECLARE_DELEGATE_RetVal_NineParams(ReturnValueType, DelegateName, Param1Type, Param2Type, Param3Type, Param4Type, Param5Type, Param6Type, Param7Type, Param8Type, Param9Type) UE_PRIVATE_DECLARE_DELEGATE(DelegateName, ReturnValueType, Param1Type, Param2Type, Param3Type, Param4Type, Param5Type, Param6Type, Param7Type, Param8Type, Param9Type)
-#define DECLARE_DYNAMIC_DELEGATE_RetVal_NineParams(ReturnValueType, DelegateName, Param1Type, Param1Name, Param2Type, Param2Name, Param3Type, Param3Name, Param4Type, Param4Name, Param5Type, Param5Name, Param6Type, Param6Name, Param7Type, Param7Name, Param8Type, Param8Name, Param9Type, Param9Name) UE_PRIVATE_DECLARE_DYNAMIC_DELEGATE(DelegateName, ReturnValueType, Param1Type, Param2Type, Param3Type, Param4Type, Param5Type, Param6Type, Param7Type, Param8Type, Param9Type)
+#define DECLARE_DELEGATE_OneParam(DelegateName, Param1Type) UE_PRIVATE_DECLARE_DELEGATE(DelegateName, void(Param1Type))
+#define DECLARE_MULTICAST_DELEGATE_OneParam(DelegateName, Param1Type) UE_PRIVATE_DECLARE_MULTICAST_DELEGATE(DelegateName, void(Param1Type))
+#define DECLARE_TS_MULTICAST_DELEGATE_OneParam(DelegateName, Param1Type) UE_PRIVATE_DECLARE_TS_MULTICAST_DELEGATE(DelegateName, void(Param1Type))
+#define DECLARE_EVENT_OneParam(OwningType, EventName, Param1Type) UE_PRIVATE_DECLARE_EVENT(OwningType, EventName, void(Param1Type))
+#define DECLARE_DYNAMIC_DELEGATE_OneParam(DelegateName, Param1Type, Param1Name) UE_PRIVATE_DECLARE_DYNAMIC_DELEGATE(DelegateName, void(Param1Type Param1Name))
+#define DECLARE_DYNAMIC_MULTICAST_DELEGATE_OneParam(DelegateName, Param1Type, Param1Name) UE_PRIVATE_DECLARE_DYNAMIC_MULTICAST_DELEGATE(DelegateName, void(Param1Type Param1Name))
+#define DECLARE_DELEGATE_RetVal_OneParam(ReturnValueType, DelegateName, Param1Type) UE_PRIVATE_DECLARE_DELEGATE(DelegateName, ReturnValueType(Param1Type))
+#define DECLARE_DYNAMIC_DELEGATE_RetVal_OneParam(ReturnValueType, DelegateName, Param1Type, Param1Name) UE_PRIVATE_DECLARE_DYNAMIC_DELEGATE(DelegateName, ReturnValueType(Param1Type Param1Name))
+
+#define DECLARE_DELEGATE_TwoParams(DelegateName, Param1Type, Param2Type) UE_PRIVATE_DECLARE_DELEGATE(DelegateName, void(Param1Type, Param2Type))
+#define DECLARE_MULTICAST_DELEGATE_TwoParams(DelegateName, Param1Type, Param2Type) UE_PRIVATE_DECLARE_MULTICAST_DELEGATE(DelegateName, void(Param1Type, Param2Type))
+#define DECLARE_TS_MULTICAST_DELEGATE_TwoParams(DelegateName, Param1Type, Param2Type) UE_PRIVATE_DECLARE_TS_MULTICAST_DELEGATE(DelegateName, void(Param1Type, Param2Type))
+#define DECLARE_EVENT_TwoParams(OwningType, EventName, Param1Type, Param2Type) UE_PRIVATE_DECLARE_EVENT(OwningType, EventName, void(Param1Type, Param2Type))
+#define DECLARE_DYNAMIC_DELEGATE_TwoParams(DelegateName, Param1Type, Param1Name, Param2Type, Param2Name) UE_PRIVATE_DECLARE_DYNAMIC_DELEGATE(DelegateName, void(Param1Type Param1Name, Param2Type Param2Name))
+#define DECLARE_DYNAMIC_MULTICAST_DELEGATE_TwoParams(DelegateName, Param1Type, Param1Name, Param2Type, Param2Name) UE_PRIVATE_DECLARE_DYNAMIC_MULTICAST_DELEGATE(DelegateName, void(Param1Type Param1Name, Param2Type Param2Name))
+#define DECLARE_DELEGATE_RetVal_TwoParams(ReturnValueType, DelegateName, Param1Type, Param2Type) UE_PRIVATE_DECLARE_DELEGATE(DelegateName, ReturnValueType(Param1Type, Param2Type))
+#define DECLARE_DYNAMIC_DELEGATE_RetVal_TwoParams(ReturnValueType, DelegateName, Param1Type, Param1Name, Param2Type, Param2Name) UE_PRIVATE_DECLARE_DYNAMIC_DELEGATE(DelegateName, ReturnValueType(Param1Type Param1Name, Param2Type Param2Name))
+
+#define DECLARE_DELEGATE_ThreeParams(DelegateName, Param1Type, Param2Type, Param3Type) UE_PRIVATE_DECLARE_DELEGATE(DelegateName, void(Param1Type, Param2Type, Param3Type))
+#define DECLARE_MULTICAST_DELEGATE_ThreeParams(DelegateName, Param1Type, Param2Type, Param3Type) UE_PRIVATE_DECLARE_MULTICAST_DELEGATE(DelegateName, void(Param1Type, Param2Type, Param3Type))
+#define DECLARE_TS_MULTICAST_DELEGATE_ThreeParams(DelegateName, Param1Type, Param2Type, Param3Type) UE_PRIVATE_DECLARE_TS_MULTICAST_DELEGATE(DelegateName, void(Param1Type, Param2Type, Param3Type))
+#define DECLARE_EVENT_ThreeParams(OwningType, EventName, Param1Type, Param2Type, Param3Type) UE_PRIVATE_DECLARE_EVENT(OwningType, EventName, void(Param1Type, Param2Type, Param3Type))
+#define DECLARE_DYNAMIC_DELEGATE_ThreeParams(DelegateName, Param1Type, Param1Name, Param2Type, Param2Name, Param3Type, Param3Name) UE_PRIVATE_DECLARE_DYNAMIC_DELEGATE(DelegateName, void(Param1Type Param1Name, Param2Type Param2Name, Param3Type Param3Name))
+#define DECLARE_DYNAMIC_MULTICAST_DELEGATE_ThreeParams(DelegateName, Param1Type, Param1Name, Param2Type, Param2Name, Param3Type, Param3Name) UE_PRIVATE_DECLARE_DYNAMIC_MULTICAST_DELEGATE(DelegateName, void(Param1Type Param1Name, Param2Type Param2Name, Param3Type Param3Name))
+#define DECLARE_DELEGATE_RetVal_ThreeParams(ReturnValueType, DelegateName, Param1Type, Param2Type, Param3Type) UE_PRIVATE_DECLARE_DELEGATE(DelegateName, ReturnValueType(Param1Type, Param2Type, Param3Type))
+#define DECLARE_DYNAMIC_DELEGATE_RetVal_ThreeParams(ReturnValueType, DelegateName, Param1Type, Param1Name, Param2Type, Param2Name, Param3Type, Param3Name) UE_PRIVATE_DECLARE_DYNAMIC_DELEGATE(DelegateName, ReturnValueType(Param1Type Param1Name, Param2Type Param2Name, Param3Type Param3Name))
+
+#define DECLARE_DELEGATE_FourParams(DelegateName, Param1Type, Param2Type, Param3Type, Param4Type) UE_PRIVATE_DECLARE_DELEGATE( DelegateName, void(Param1Type, Param2Type, Param3Type, Param4Type))
+#define DECLARE_MULTICAST_DELEGATE_FourParams(DelegateName, Param1Type, Param2Type, Param3Type, Param4Type) UE_PRIVATE_DECLARE_MULTICAST_DELEGATE(DelegateName, void(Param1Type, Param2Type, Param3Type, Param4Type))
+#define DECLARE_TS_MULTICAST_DELEGATE_FourParams(DelegateName, Param1Type, Param2Type, Param3Type, Param4Type) UE_PRIVATE_DECLARE_TS_MULTICAST_DELEGATE(DelegateName, void(Param1Type, Param2Type, Param3Type, Param4Type))
+#define DECLARE_EVENT_FourParams(OwningType, EventName, Param1Type, Param2Type, Param3Type, Param4Type) UE_PRIVATE_DECLARE_EVENT(OwningType, EventName, void(Param1Type, Param2Type, Param3Type, Param4Type))
+#define DECLARE_DYNAMIC_DELEGATE_FourParams(DelegateName, Param1Type, Param1Name, Param2Type, Param2Name, Param3Type, Param3Name, Param4Type, Param4Name) UE_PRIVATE_DECLARE_DYNAMIC_DELEGATE(DelegateName, void(Param1Type Param1Name, Param2Type Param2Name, Param3Type Param3Name, Param4Type Param4Name))
+#define DECLARE_DYNAMIC_MULTICAST_DELEGATE_FourParams(DelegateName, Param1Type, Param1Name, Param2Type, Param2Name, Param3Type, Param3Name, Param4Type, Param4Name) UE_PRIVATE_DECLARE_DYNAMIC_MULTICAST_DELEGATE(DelegateName, void(Param1Type Param1Name, Param2Type Param2Name, Param3Type Param3Name, Param4Type Param4Name))
+#define DECLARE_DELEGATE_RetVal_FourParams(ReturnValueType, DelegateName, Param1Type, Param2Type, Param3Type, Param4Type) UE_PRIVATE_DECLARE_DELEGATE(DelegateName, ReturnValueType(Param1Type, Param2Type, Param3Type, Param4Type))
+#define DECLARE_DYNAMIC_DELEGATE_RetVal_FourParams(ReturnValueType, DelegateName, Param1Type, Param1Name, Param2Type, Param2Name, Param3Type, Param3Name, Param4Type, Param4Name) UE_PRIVATE_DECLARE_DYNAMIC_DELEGATE(DelegateName, ReturnValueType(Param1Type Param1Name, Param2Type Param2Name, Param3Type Param3Name, Param4Type Param4Name))
+
+#define DECLARE_DELEGATE_FiveParams(DelegateName, Param1Type, Param2Type, Param3Type, Param4Type, Param5Type) UE_PRIVATE_DECLARE_DELEGATE(DelegateName, void(Param1Type, Param2Type, Param3Type, Param4Type, Param5Type))
+#define DECLARE_MULTICAST_DELEGATE_FiveParams(DelegateName, Param1Type, Param2Type, Param3Type, Param4Type, Param5Type) UE_PRIVATE_DECLARE_MULTICAST_DELEGATE(DelegateName, void(Param1Type, Param2Type, Param3Type, Param4Type, Param5Type))
+#define DECLARE_TS_MULTICAST_DELEGATE_FiveParams(DelegateName, Param1Type, Param2Type, Param3Type, Param4Type, Param5Type) UE_PRIVATE_DECLARE_TS_MULTICAST_DELEGATE(DelegateName, void(Param1Type, Param2Type, Param3Type, Param4Type, Param5Type))
+#define DECLARE_EVENT_FiveParams(OwningType, EventName, Param1Type, Param2Type, Param3Type, Param4Type, Param5Type) UE_PRIVATE_DECLARE_EVENT(OwningType, EventName, void(Param1Type, Param2Type, Param3Type, Param4Type, Param5Type))
+#define DECLARE_DYNAMIC_DELEGATE_FiveParams(DelegateName, Param1Type, Param1Name, Param2Type, Param2Name, Param3Type, Param3Name, Param4Type, Param4Name, Param5Type, Param5Name) UE_PRIVATE_DECLARE_DYNAMIC_DELEGATE(DelegateName, void(Param1Type Param1Name, Param2Type Param2Name, Param3Type Param3Name, Param4Type Param4Name, Param5Type Param5Name))
+#define DECLARE_DYNAMIC_MULTICAST_DELEGATE_FiveParams(DelegateName, Param1Type, Param1Name, Param2Type, Param2Name, Param3Type, Param3Name, Param4Type, Param4Name, Param5Type, Param5Name) UE_PRIVATE_DECLARE_DYNAMIC_MULTICAST_DELEGATE(DelegateName, void(Param1Type Param1Name, Param2Type Param2Name, Param3Type Param3Name, Param4Type Param4Name, Param5Type Param5Name))
+#define DECLARE_DELEGATE_RetVal_FiveParams(ReturnValueType, DelegateName, Param1Type, Param2Type, Param3Type, Param4Type, Param5Type) UE_PRIVATE_DECLARE_DELEGATE(DelegateName, ReturnValueType(Param1Type, Param2Type, Param3Type, Param4Type, Param5Type))
+#define DECLARE_DYNAMIC_DELEGATE_RetVal_FiveParams(ReturnValueType, DelegateName, Param1Type, Param1Name, Param2Type, Param2Name, Param3Type, Param3Name, Param4Type, Param4Name, Param5Type, Param5Name) UE_PRIVATE_DECLARE_DYNAMIC_DELEGATE(DelegateName, ReturnValueType(Param1Type Param1Name, Param2Type Param2Name, Param3Type Param3Name, Param4Type Param4Name, Param5Type Param5Name))
+
+#define DECLARE_DELEGATE_SixParams(DelegateName, Param1Type, Param2Type, Param3Type, Param4Type, Param5Type, Param6Type) UE_PRIVATE_DECLARE_DELEGATE(DelegateName, void(Param1Type, Param2Type, Param3Type, Param4Type, Param5Type, Param6Type))
+#define DECLARE_MULTICAST_DELEGATE_SixParams(DelegateName, Param1Type, Param2Type, Param3Type, Param4Type, Param5Type, Param6Type) UE_PRIVATE_DECLARE_MULTICAST_DELEGATE(DelegateName, void(Param1Type, Param2Type, Param3Type, Param4Type, Param5Type, Param6Type))
+#define DECLARE_TS_MULTICAST_DELEGATE_SixParams(DelegateName, Param1Type, Param2Type, Param3Type, Param4Type, Param5Type, Param6Type) UE_PRIVATE_DECLARE_TS_MULTICAST_DELEGATE(DelegateName, void(Param1Type, Param2Type, Param3Type, Param4Type, Param5Type, Param6Type))
+#define DECLARE_EVENT_SixParams(OwningType, EventName, Param1Type, Param2Type, Param3Type, Param4Type, Param5Type, Param6Type) UE_PRIVATE_DECLARE_EVENT(OwningType, EventName, void(Param1Type, Param2Type, Param3Type, Param4Type, Param5Type, Param6Type))
+#define DECLARE_DYNAMIC_DELEGATE_SixParams(DelegateName, Param1Type, Param1Name, Param2Type, Param2Name, Param3Type, Param3Name, Param4Type, Param4Name, Param5Type, Param5Name, Param6Type, Param6Name) UE_PRIVATE_DECLARE_DYNAMIC_DELEGATE(DelegateName, void(Param1Type Param1Name, Param2Type Param2Name, Param3Type Param3Name, Param4Type Param4Name, Param5Type Param5Name, Param6Type Param6Name))
+#define DECLARE_DYNAMIC_MULTICAST_DELEGATE_SixParams(DelegateName, Param1Type, Param1Name, Param2Type, Param2Name, Param3Type, Param3Name, Param4Type, Param4Name, Param5Type, Param5Name, Param6Type, Param6Name) UE_PRIVATE_DECLARE_DYNAMIC_MULTICAST_DELEGATE(DelegateName, void(Param1Type Param1Name, Param2Type Param2Name, Param3Type Param3Name, Param4Type Param4Name, Param5Type Param5Name, Param6Type Param6Name))
+#define DECLARE_DELEGATE_RetVal_SixParams(ReturnValueType, DelegateName, Param1Type, Param2Type, Param3Type, Param4Type, Param5Type, Param6Type) UE_PRIVATE_DECLARE_DELEGATE(DelegateName, ReturnValueType(Param1Type, Param2Type, Param3Type, Param4Type, Param5Type, Param6Type))
+#define DECLARE_DYNAMIC_DELEGATE_RetVal_SixParams(ReturnValueType, DelegateName, Param1Type, Param1Name, Param2Type, Param2Name, Param3Type, Param3Name, Param4Type, Param4Name, Param5Type, Param5Name, Param6Type, Param6Name) UE_PRIVATE_DECLARE_DYNAMIC_DELEGATE(DelegateName, ReturnValueType(Param1Type Param1Name, Param2Type Param2Name, Param3Type Param3Name, Param4Type Param4Name, Param5Type Param5Name, Param6Type Param6Name))
+
+#define DECLARE_DELEGATE_SevenParams(DelegateName, Param1Type, Param2Type, Param3Type, Param4Type, Param5Type, Param6Type, Param7Type) UE_PRIVATE_DECLARE_DELEGATE(DelegateName, void(Param1Type, Param2Type, Param3Type, Param4Type, Param5Type, Param6Type, Param7Type))
+#define DECLARE_MULTICAST_DELEGATE_SevenParams(DelegateName, Param1Type, Param2Type, Param3Type, Param4Type, Param5Type, Param6Type, Param7Type) UE_PRIVATE_DECLARE_MULTICAST_DELEGATE(DelegateName, void(Param1Type, Param2Type, Param3Type, Param4Type, Param5Type, Param6Type, Param7Type))
+#define DECLARE_TS_MULTICAST_DELEGATE_SevenParams(DelegateName, Param1Type, Param2Type, Param3Type, Param4Type, Param5Type, Param6Type, Param7Type) UE_PRIVATE_DECLARE_TS_MULTICAST_DELEGATE(DelegateName, void(Param1Type, Param2Type, Param3Type, Param4Type, Param5Type, Param6Type, Param7Type))
+#define DECLARE_EVENT_SevenParams(OwningType, EventName, Param1Type, Param2Type, Param3Type, Param4Type, Param5Type, Param6Type, Param7Type) UE_PRIVATE_DECLARE_EVENT(OwningType, EventName, void(Param1Type, Param2Type, Param3Type, Param4Type, Param5Type, Param6Type, Param7Type))
+#define DECLARE_DYNAMIC_DELEGATE_SevenParams(DelegateName, Param1Type, Param1Name, Param2Type, Param2Name, Param3Type, Param3Name, Param4Type, Param4Name, Param5Type, Param5Name, Param6Type, Param6Name, Param7Type, Param7Name) UE_PRIVATE_DECLARE_DYNAMIC_DELEGATE(DelegateName, void(Param1Type Param1Name, Param2Type Param2Name, Param3Type Param3Name, Param4Type Param4Name, Param5Type Param5Name, Param6Type Param6Name, Param7Type Param7Name))
+#define DECLARE_DYNAMIC_MULTICAST_DELEGATE_SevenParams(DelegateName, Param1Type, Param1Name, Param2Type, Param2Name, Param3Type, Param3Name, Param4Type, Param4Name, Param5Type, Param5Name, Param6Type, Param6Name, Param7Type, Param7Name) UE_PRIVATE_DECLARE_DYNAMIC_MULTICAST_DELEGATE(DelegateName, void(Param1Type Param1Name, Param2Type Param2Name, Param3Type Param3Name, Param4Type Param4Name, Param5Type Param5Name, Param6Type Param6Name, Param7Type Param7Name))
+#define DECLARE_DELEGATE_RetVal_SevenParams(ReturnValueType, DelegateName, Param1Type, Param2Type, Param3Type, Param4Type, Param5Type, Param6Type, Param7Type) UE_PRIVATE_DECLARE_DELEGATE(DelegateName, ReturnValueType(Param1Type, Param2Type, Param3Type, Param4Type, Param5Type, Param6Type, Param7Type))
+#define DECLARE_DYNAMIC_DELEGATE_RetVal_SevenParams(ReturnValueType, DelegateName, Param1Type, Param1Name, Param2Type, Param2Name, Param3Type, Param3Name, Param4Type, Param4Name, Param5Type, Param5Name, Param6Type, Param6Name, Param7Type, Param7Name) UE_PRIVATE_DECLARE_DYNAMIC_DELEGATE(DelegateName, ReturnValueType(Param1Type Param1Name, Param2Type Param2Name, Param3Type Param3Name, Param4Type Param4Name, Param5Type Param5Name, Param6Type Param6Name, Param7Type Param7Name))
+
+#define DECLARE_DELEGATE_EightParams(DelegateName, Param1Type, Param2Type, Param3Type, Param4Type, Param5Type, Param6Type, Param7Type, Param8Type) UE_PRIVATE_DECLARE_DELEGATE(DelegateName, void(Param1Type, Param2Type, Param3Type, Param4Type, Param5Type, Param6Type, Param7Type, Param8Type))
+#define DECLARE_MULTICAST_DELEGATE_EightParams(DelegateName, Param1Type, Param2Type, Param3Type, Param4Type, Param5Type, Param6Type, Param7Type, Param8Type) UE_PRIVATE_DECLARE_MULTICAST_DELEGATE(DelegateName, void(Param1Type, Param2Type, Param3Type, Param4Type, Param5Type, Param6Type, Param7Type, Param8Type))
+#define DECLARE_TS_MULTICAST_DELEGATE_EightParams(DelegateName, Param1Type, Param2Type, Param3Type, Param4Type, Param5Type, Param6Type, Param7Type, Param8Type) UE_PRIVATE_DECLARE_TS_MULTICAST_DELEGATE(DelegateName, void(Param1Type, Param2Type, Param3Type, Param4Type, Param5Type, Param6Type, Param7Type, Param8Type))
+#define DECLARE_EVENT_EightParams(OwningType, EventName, Param1Type, Param2Type, Param3Type, Param4Type, Param5Type, Param6Type, Param7Type, Param8Type) UE_PRIVATE_DECLARE_EVENT(OwningType, EventName, void(Param1Type, Param2Type, Param3Type, Param4Type, Param5Type, Param6Type, Param7Type, Param8Type))
+#define DECLARE_DYNAMIC_DELEGATE_EightParams(DelegateName, Param1Type, Param1Name, Param2Type, Param2Name, Param3Type, Param3Name, Param4Type, Param4Name, Param5Type, Param5Name, Param6Type, Param6Name, Param7Type, Param7Name, Param8Type, Param8Name) UE_PRIVATE_DECLARE_DYNAMIC_DELEGATE(DelegateName, void(Param1Type Param1Name, Param2Type Param2Name, Param3Type Param3Name, Param4Type Param5Name, Param5Type Param5Name, Param6Type Param6Name, Param7Type Param7Name, Param8Type Param8Name))
+#define DECLARE_DYNAMIC_MULTICAST_DELEGATE_EightParams(DelegateName, Param1Type, Param1Name, Param2Type, Param2Name, Param3Type, Param3Name, Param4Type, Param4Name, Param5Type, Param5Name, Param6Type, Param6Name, Param7Type, Param7Name, Param8Type, Param8Name) UE_PRIVATE_DECLARE_DYNAMIC_MULTICAST_DELEGATE(DelegateName, void(Param1Type Param1Name, Param2Type Param2Name, Param3Type Param3Name, Param4Type Param4Name, Param5Type Param5Name, Param6Type Param6Name, Param7Type Param7Name, Param8Type Param8Name))
+#define DECLARE_DELEGATE_RetVal_EightParams(ReturnValueType, DelegateName, Param1Type, Param2Type, Param3Type, Param4Type, Param5Type, Param6Type, Param7Type, Param8Type) UE_PRIVATE_DECLARE_DELEGATE(DelegateName, ReturnValueType(Param1Type, Param2Type, Param3Type, Param4Type, Param5Type, Param6Type, Param7Type, Param8Type))
+#define DECLARE_DYNAMIC_DELEGATE_RetVal_EightParams(ReturnValueType, DelegateName, Param1Type, Param1Name, Param2Type, Param2Name, Param3Type, Param3Name, Param4Type, Param4Name, Param5Type, Param5Name, Param6Type, Param6Name, Param7Type, Param7Name, Param8Type, Param8Name) UE_PRIVATE_DECLARE_DYNAMIC_DELEGATE(DelegateName, ReturnValueType(Param1Type Param1Name, Param2Type Param2Name, Param3Type Param3Name, Param4Type Param4Name, Param5Type Param5Name, Param6Type Param6Name, Param7Type Param7Name, Param8Type Param8Name))
+
+#define DECLARE_DELEGATE_NineParams(DelegateName, Param1Type, Param2Type, Param3Type, Param4Type, Param5Type, Param6Type, Param7Type, Param8Type, Param9Type) UE_PRIVATE_DECLARE_DELEGATE(DelegateName, void(Param1Type, Param2Type, Param3Type, Param4Type, Param5Type, Param6Type, Param7Type, Param8Type, Param9Type))
+#define DECLARE_MULTICAST_DELEGATE_NineParams(DelegateName, Param1Type, Param2Type, Param3Type, Param4Type, Param5Type, Param6Type, Param7Type, Param8Type, Param9Type) UE_PRIVATE_DECLARE_MULTICAST_DELEGATE(DelegateName, void(Param1Type, Param2Type, Param3Type, Param4Type, Param5Type, Param6Type, Param7Type, Param8Type, Param9Type))
+#define DECLARE_TS_MULTICAST_DELEGATE_NineParams(DelegateName, Param1Type, Param2Type, Param3Type, Param4Type, Param5Type, Param6Type, Param7Type, Param8Type, Param9Type) UE_PRIVATE_DECLARE_TS_MULTICAST_DELEGATE(DelegateName, void(Param1Type, Param2Type, Param3Type, Param4Type, Param5Type, Param6Type, Param7Type, Param8Type, Param9Type))
+#define DECLARE_EVENT_NineParams(OwningType, EventName, Param1Type, Param2Type, Param3Type, Param4Type, Param5Type, Param6Type, Param7Type, Param8Type, Param9Type) UE_PRIVATE_DECLARE_EVENT(OwningType, EventName, void(Param1Type, Param2Type, Param3Type, Param4Type, Param5Type, Param6Type, Param7Type, Param8Type, Param9Type))
+#define DECLARE_DYNAMIC_DELEGATE_NineParams(DelegateName, Param1Type, Param1Name, Param2Type, Param2Name, Param3Type, Param3Name, Param4Type, Param4Name, Param5Type, Param5Name, Param6Type, Param6Name, Param7Type, Param7Name, Param8Type, Param8Name, Param9Type, Param9Name) UE_PRIVATE_DECLARE_DYNAMIC_DELEGATE(DelegateName, void(Param1Type Param1Name, Param2Type Param2Name, Param3Type Param3Name, Param4Type Param4Name, Param5Type Param5Name, Param6Type Param6Name, Param7Type Param7Name, Param8Type Param8Name, Param9Type Param9Name))
+#define DECLARE_DYNAMIC_MULTICAST_DELEGATE_NineParams(DelegateName, Param1Type, Param1Name, Param2Type, Param2Name, Param3Type, Param3Name, Param4Type, Param4Name, Param5Type, Param5Name, Param6Type, Param6Name, Param7Type, Param7Name, Param8Type, Param8Name, Param9Type, Param9Name) UE_PRIVATE_DECLARE_DYNAMIC_MULTICAST_DELEGATE(DelegateName, void(Param1Type Param1Name, Param2Type Param2Name, Param3Type Param3Name, Param4Type Param4Name, Param5Type Param5Name, Param6Type Param6Name, Param7Type Param7Name, Param8Type Param8Name, Param9Type Param9Name))
+#define DECLARE_DELEGATE_RetVal_NineParams(ReturnValueType, DelegateName, Param1Type, Param2Type, Param3Type, Param4Type, Param5Type, Param6Type, Param7Type, Param8Type, Param9Type) UE_PRIVATE_DECLARE_DELEGATE( DelegateName, ReturnValueType(Param1Type, Param2Type, Param3Type, Param4Type, Param5Type, Param6Type, Param7Type, Param8Type, Param9Type))
+#define DECLARE_DYNAMIC_DELEGATE_RetVal_NineParams(ReturnValueType, DelegateName, Param1Type, Param1Name, Param2Type, Param2Name, Param3Type, Param3Name, Param4Type, Param4Name, Param5Type, Param5Name, Param6Type, Param6Name, Param7Type, Param7Name, Param8Type, Param8Name, Param9Type, Param9Name) UE_PRIVATE_DECLARE_DYNAMIC_DELEGATE(DelegateName, ReturnValueType(Param1Type Param1Name, Param2Type Param2Name, Param3Type Param3Name, Param4Type Param4Name, Param5Type Param5Name, Param6Type Param6Name, Param7Type Param7Name, Param8Type Param8Name, Param9Type Param9Name))
diff --git a/Engine/Source/Runtime/Core/Public/Delegates/DelegateSignatureImpl.inl b/Engine/Source/Runtime/Core/Public/Delegates/DelegateSignatureImpl.inl
--- a/Engine/Source/Runtime/Core/Public/Delegates/DelegateSignatureImpl.inl
+++ b/Engine/Source/Runtime/Core/Public/Delegates/DelegateSignatureImpl.inl
@@ -1143,8 +1143,8 @@ template <typename DelegateSignature>
 using TTSMulticastDelegate = TMulticastDelegate<DelegateSignature, FDefaultTSDelegateUserPolicy>;
 
 
-template <typename ThreadSafetyMode, typename RetValType, typename... ParamTypes>
-class TBaseDynamicMulticastDelegate;
+template <typename FuncType, typename ThreadSafetyMode = FNotThreadSafeDelegateMode>
+class TDynamicMulticastDelegate;
 
 template <typename MulticastDelegate, typename OwningClass, typename DelegateInfoClass>
 struct TSparseDynamicDelegate;
@@ -1154,13 +1154,16 @@ struct TSparseDynamicDelegate;
  * macros to create the actual delegate type, templated to the function signature the delegate is compatible with.
  * Then, you can create an instance of that class when you want to assign functions to the delegate.
  */
-template <typename ThreadSafetyMode, typename RetValType, typename... ParamTypes>
-class TBaseDynamicDelegate : public TScriptDelegate<ThreadSafetyMode>
+template <typename FuncType, typename ThreadSafetyMode = FNotThreadSafeDelegateMode>
+class TDynamicDelegate;
+
+template <typename RetValType, typename... ParamTypes, typename ThreadSafetyMode>
+class TDynamicDelegate<RetValType(ParamTypes...), ThreadSafetyMode> : public TScriptDelegate<ThreadSafetyMode>
 {
-	template <typename, typename, typename...>
-	friend class TBaseDynamicMulticastDelegate;
+	template <typename, typename>
+	friend class TDynamicMulticastDelegate;
 
-	/** Sparse version of TBaseDynamicDelegate */
+	/** Sparse version of TDynamicDelegate */
 	template <typename, typename, typename>
 	friend struct TSparseDynamicDelegate;
 
@@ -1168,16 +1171,19 @@ public:
 	/**
 	 * Default constructor
 	 */
-	TBaseDynamicDelegate() { }
+	TDynamicDelegate()
+	{
+	}
 
 	/**
 	 * Construction from an FScriptDelegate must be explicit.  This is really only used by UObject system internals.
 	 *
 	 * @param	InScriptDelegate	The delegate to construct from by copying
 	 */
-	explicit TBaseDynamicDelegate( const TScriptDelegate<ThreadSafetyMode>& InScriptDelegate )
-		: TScriptDelegate<ThreadSafetyMode>( InScriptDelegate )
-	{ }
+	explicit TDynamicDelegate(const TScriptDelegate<ThreadSafetyMode>& InScriptDelegate)
+		: TScriptDelegate<ThreadSafetyMode>(InScriptDelegate)
+	{
+	}
 
 	/**
 	 * Templated helper class to define a typedef for user's method pointer, then used below
@@ -1278,26 +1284,29 @@ public:
  * signature the delegate is compatible with.   Then, you can create an instance of that class when you
  * want to assign functions to the delegate.
  */
-template <typename ThreadSafetyMode, typename RetValType, typename... ParamTypes>
-class TBaseDynamicMulticastDelegate : public TMulticastScriptDelegate<ThreadSafetyMode>
+template <typename RetValType, typename... ParamTypes, typename ThreadSafetyMode>
+class TDynamicMulticastDelegate<RetValType(ParamTypes...), ThreadSafetyMode> : public TMulticastScriptDelegate<ThreadSafetyMode>
 {
 public:
 	/** The actual single-cast delegate class for this multi-cast delegate */
-	typedef TBaseDynamicDelegate<ThreadSafetyMode, RetValType, ParamTypes...> FDelegate;
+	using FDelegate = TDynamicDelegate<RetValType(ParamTypes...), ThreadSafetyMode>;
 
 	/**
 	 * Default constructor
 	 */
-	TBaseDynamicMulticastDelegate() { }
+	TDynamicMulticastDelegate()
+	{
+	}
 
 	/**
 	 * Construction from an FMulticastScriptDelegate must be explicit.  This is really only used by UObject system internals.
 	 *
 	 * @param	InScriptDelegate	The delegate to construct from by copying
 	 */
-	explicit TBaseDynamicMulticastDelegate( const TMulticastScriptDelegate<ThreadSafetyMode>& InMulticastScriptDelegate )
-		: TMulticastScriptDelegate<ThreadSafetyMode>( InMulticastScriptDelegate )
-	{ }
+	explicit TDynamicMulticastDelegate(const TMulticastScriptDelegate<ThreadSafetyMode>& InMulticastScriptDelegate)
+		: TMulticastScriptDelegate<ThreadSafetyMode>(InMulticastScriptDelegate)
+	{
+	}
 
 private:
 	template <typename UserClass, typename... VarTypes>
diff --git a/Engine/Source/Runtime/CoreUObject/Public/UObject/SparseDelegate.h b/Engine/Source/Runtime/CoreUObject/Public/UObject/SparseDelegate.h
--- a/Engine/Source/Runtime/CoreUObject/Public/UObject/SparseDelegate.h
+++ b/Engine/Source/Runtime/CoreUObject/Public/UObject/SparseDelegate.h
@@ -192,7 +192,7 @@ protected:
 	bool bIsBound;
 };
 
-/** Sparse version of TBaseDynamicDelegate */
+/** Sparse version of TDynamicDelegate */
 template <typename MulticastDelegate, typename OwningClass, typename DelegateInfoClass>
 struct TSparseDynamicDelegate : public FSparseDelegate
 {
