[Core] Move `ResolvePackageShaderMaps` delegate from FCoreInternalDelegates into FCoreDelegates as this is needed in the public interface of the shader code library.

#rnx

#Synced-CL 50198447

diff --git a/Engine/Source/Runtime/Core/Internal/Misc/CoreDelegatesInternal.h b/Engine/Source/Runtime/Core/Internal/Misc/CoreDelegatesInternal.h
--- a/Engine/Source/Runtime/Core/Internal/Misc/CoreDelegatesInternal.h
+++ b/Engine/Source/Runtime/Core/Internal/Misc/CoreDelegatesInternal.h
@@ -41,16 +41,4 @@ struct FCoreInternalDelegates
 	static CORE_API FMountOperationPak& GetOnPakMountOperation();
 
 	static CORE_API FCurrentlyMountedPaksDelegate& GetCurrentlyMountedPaksDelegate();
-
-	/** Callback triggered when a shader map hash has been resolved to an I/O store chunk ID. */
-	using FIoChunkIdResolvedFunc = TFunctionRef<void(const class FIoChunkId&)>;
-
-	/**
-	 * Callback to resolve I/O store chunk IDs from shader code hashes.
-	 * Used when installing content on-demand.
-	 *
-	 * @param	Hashes		Array of hashes that should be resolved.
-	 * @param	Callback	Callback triggered for reach resolved chunk ID.
-	 */
-	static CORE_API TDelegate<void(TArrayView<const FSHAHash>, FIoChunkIdResolvedFunc)> ResolvePackageShaderMaps;
 };
diff --git a/Engine/Source/Runtime/Core/Private/Misc/CoreDelegates.cpp b/Engine/Source/Runtime/Core/Private/Misc/CoreDelegates.cpp
--- a/Engine/Source/Runtime/Core/Private/Misc/CoreDelegates.cpp
+++ b/Engine/Source/Runtime/Core/Private/Misc/CoreDelegates.cpp
@@ -79,8 +79,6 @@ FCurrentlyMountedPaksDelegate& FCoreInternalDelegates::GetCurrentlyMountedPaksDe
 	return Delegate;
 }
 
-TDelegate<void(TArrayView<const FSHAHash>, FCoreInternalDelegates::FIoChunkIdResolvedFunc)> FCoreInternalDelegates::ResolvePackageShaderMaps;
-
 //} // namespace UE::Core
 
 TMulticastDelegate<void(bool, int32, int32)> FCoreDelegates::OnUserLoginChangedEvent;
@@ -316,6 +314,8 @@ TTSMulticastDelegate<void(FCoreDelegates::FExtraBinaryConfigData&)>& FCoreDelega
 	return Singleton;
 }
 
+
+TDelegate<void(TArrayView<const FSHAHash>, FCoreDelegates::FIoChunkIdResolvedFunc)> FCoreDelegates::ResolvePackageShaderMaps;
 TDelegate<void(TArrayView<const FSHAHash>, FCoreDelegates::FAttachShaderReadRequestFunc)> FCoreDelegates::PreloadPackageShaderMaps;
 TDelegate<void(TArrayView<const FSHAHash>)> FCoreDelegates::ReleasePreloadedPackageShaderMaps;
 TMulticastDelegate<void(const FLogCategoryName&, ELogVerbosity::Type, ELogVerbosity::Type)> FCoreDelegates::OnLogVerbosityChanged;
diff --git a/Engine/Source/Runtime/Core/Public/Misc/CoreDelegates.h b/Engine/Source/Runtime/Core/Public/Misc/CoreDelegates.h
--- a/Engine/Source/Runtime/Core/Public/Misc/CoreDelegates.h
+++ b/Engine/Source/Runtime/Core/Public/Misc/CoreDelegates.h
@@ -586,6 +586,19 @@ public:
 		}
 	};
 	
+
+	/** Callback triggered when a shader map hash has been resolved to an I/O store chunk ID. */
+	using FIoChunkIdResolvedFunc = TFunctionRef<void(const class FIoChunkId&)>;
+
+	/**
+	 * Callback to resolve I/O store chunk IDs from shader code hashes.
+	 * Used when installing content on-demand.
+	 *
+	 * @param	Hashes		Array of hashes that should be resolved.
+	 * @param	Callback	Callback triggered for each resolved chunk ID.
+	 */
+	static CORE_API TDelegate<void(TArrayView<const FSHAHash>, FIoChunkIdResolvedFunc)> ResolvePackageShaderMaps;
+
 	static CORE_API TTSMulticastDelegate<void(FExtraBinaryConfigData&)>& TSAccessExtraBinaryConfigData();
 
 	using FAttachShaderReadRequestFunc = TFunctionRef<class FIoRequest(const class FIoChunkId&, FGraphEventRef)>;
diff --git a/Engine/Source/Runtime/Experimental/IoStore/OnDemand/Private/OnDemandContentInstaller.cpp b/Engine/Source/Runtime/Experimental/IoStore/OnDemand/Private/OnDemandContentInstaller.cpp
--- a/Engine/Source/Runtime/Experimental/IoStore/OnDemand/Private/OnDemandContentInstaller.cpp
+++ b/Engine/Source/Runtime/Experimental/IoStore/OnDemand/Private/OnDemandContentInstaller.cpp
@@ -240,7 +240,7 @@ void ResolveChunksToInstall(
 					ResolvedPackageIds.Add(PackageId);
 					if (PackageStoreEntry.ShaderMapHashes.IsEmpty() == false)
 					{
-						FCoreInternalDelegates::ResolvePackageShaderMaps.ExecuteIfBound(
+						FCoreDelegates::ResolvePackageShaderMaps.ExecuteIfBound(
 							PackageStoreEntry.ShaderMapHashes,
 							[&ShaderMapChunkIds](const FIoChunkId& ChunkId)
 							{
