[Core]
* Moved mac semaphore implementation details to cpp file to remove sdk dependencies in header

#Synced-CL 50214969

diff --git a/Engine/Source/Runtime/Core/Private/Mac/MacSemaphore.cpp b/Engine/Source/Runtime/Core/Private/Mac/MacSemaphore.cpp
new file mode 100644
--- /dev/null
+++ b/Engine/Source/Runtime/Core/Private/Mac/MacSemaphore.cpp
@@ -0,0 +1,38 @@
+// Copyright Epic Games, Inc. All Rights Reserved.
+
+#include "Mac/MacSemaphore.h"
+#include "Misc/AssertionMacros.h"
+#include <dispatch/dispatch.h>
+
+FMacSemaphore::FMacSemaphore(int32 InitialCount, int32 /*MaxCount*/)
+	: Semaphore((dispatch_semaphore_t)dispatch_semaphore_create(InitialCount))
+{
+	checkfSlow(InitialCount >= 0, TEXT("Semaphore's initial count must be non negative value: %d"), InitialCount);
+}
+
+FMacSemaphore::~FMacSemaphore()
+{
+	dispatch_release((dispatch_semaphore_t)Semaphore);
+}
+
+void FMacSemaphore::Acquire()
+{
+	intptr_t Res = dispatch_semaphore_wait((dispatch_semaphore_t)Semaphore, DISPATCH_TIME_FOREVER);
+	checkfSlow(Res == 0, TEXT("Acquiring semaphore failed"));
+}
+
+bool FMacSemaphore::TryAcquire(FTimespan Timeout)
+{
+	dispatch_time_t TS = dispatch_time(DISPATCH_TIME_NOW, (int64)Timeout.GetTotalMicroseconds() * 1000);
+	intptr_t Res = dispatch_semaphore_wait((dispatch_semaphore_t)Semaphore, TS);
+	return Res == 0;
+}
+
+void FMacSemaphore::Release(int32 Count)
+{
+	checkfSlow(Count > 0, TEXT("Releasing semaphore with Count = %d, that should be greater than 0"), Count);
+	for (int i = 0; i != Count; ++i)
+	{
+		dispatch_semaphore_signal((dispatch_semaphore_t)Semaphore);
+	}
+}
diff --git a/Engine/Source/Runtime/Core/Public/Mac/MacSemaphore.h b/Engine/Source/Runtime/Core/Public/Mac/MacSemaphore.h
--- a/Engine/Source/Runtime/Core/Public/Mac/MacSemaphore.h
+++ b/Engine/Source/Runtime/Core/Public/Mac/MacSemaphore.h
@@ -4,50 +4,21 @@
 
 #include "CoreTypes.h"
 #include "Misc/Timespan.h"
-#include "Misc/AssertionMacros.h"
-
-#include <dispatch/dispatch.h>
 
 class FMacSemaphore
 {
 public:
 	UE_NONCOPYABLE(FMacSemaphore);
 
-	FMacSemaphore(int32 InitialCount, int32 /*MaxCount*/)
-		: Semaphore(dispatch_semaphore_create(InitialCount))
-	{
-		checkfSlow(InitialCount >= 0, TEXT("Semaphore's initial count must be non negative value: %d"), InitialCount);
-	}
-    
-    virtual ~FMacSemaphore()
-    {
-        dispatch_release(Semaphore);
-    }
-
-	void Acquire()
-	{
-		intptr_t Res = dispatch_semaphore_wait(Semaphore, DISPATCH_TIME_FOREVER);
-		checkfSlow(Res == 0, TEXT("Acquiring semaphore failed"));
-	}
-
-	bool TryAcquire(FTimespan Timeout = FTimespan::Zero())
-	{
-		dispatch_time_t TS = dispatch_time(DISPATCH_TIME_NOW, (int64)Timeout.GetTotalMicroseconds() * 1000);
-		intptr_t Res = dispatch_semaphore_wait(Semaphore, TS);
-		return Res == 0;
-	}
+	CORE_API FMacSemaphore(int32 InitialCount, int32 /*MaxCount*/);
+    CORE_API virtual ~FMacSemaphore();
 
-	void Release(int32 Count = 1)
-	{
-		checkfSlow(Count > 0, TEXT("Releasing semaphore with Count = %d, that should be greater than 0"), Count);
-		for (int i = 0; i != Count; ++i)
-		{
-			dispatch_semaphore_signal(Semaphore);
-		}
-	}
+	CORE_API void Acquire();
+	CORE_API bool TryAcquire(FTimespan Timeout = FTimespan::Zero());
+	CORE_API void Release(int32 Count = 1);
 
 private:
-	dispatch_semaphore_t Semaphore;
+	void* Semaphore; // (dispatch_semaphore_t)
 };
 
 typedef FMacSemaphore FSemaphore;
