LLM: Fix ProgramSize falling to 0 after a few publish frames of LLM. We were incorrectly pruning allocationgroups that still had a non-zero size and need to be kept for reporting.
#rnx

#Synced-CL 50149715

diff --git a/Engine/Source/Runtime/Core/Private/HAL/LLM/AllocationGroup.cpp b/Engine/Source/Runtime/Core/Private/HAL/LLM/AllocationGroup.cpp
--- a/Engine/Source/Runtime/Core/Private/HAL/LLM/AllocationGroup.cpp
+++ b/Engine/Source/Runtime/Core/Private/HAL/LLM/AllocationGroup.cpp
@@ -115,7 +115,22 @@ void FAllocationGroup::ReleaseThread()
 
 bool FAllocationGroup::IsReferenced() const
 {
-	return ActiveAllocs != 0 || ThreadRefCount.load(std::memory_order_relaxed) != 0;
+	if (ActiveAllocs != 0)
+	{
+		return true;
+	}
+	for (uint64 Size : TrackerSizes)
+	{
+		if (Size != 0)
+		{
+			return true;
+		}
+	}
+	if (ThreadRefCount.load(std::memory_order_relaxed) != 0)
+	{
+		return true;
+	}
+	return false;
 }
 
 void FAllocationGroupTrackingData::ConditionalInitialize(FAllocationGroup* InGroup, bool bReferenceTrack)
