FPaths: Add UE::Paths::EFlags, add GetExtension parameter EFlags::AllowCompoundExtension.
#rnx

#Synced-CL 50149715

diff --git a/Engine/Source/Developer/IoStoreUtilities/Private/ZenStoreWriter.cpp b/Engine/Source/Developer/IoStoreUtilities/Private/ZenStoreWriter.cpp
--- a/Engine/Source/Developer/IoStoreUtilities/Private/ZenStoreWriter.cpp
+++ b/Engine/Source/Developer/IoStoreUtilities/Private/ZenStoreWriter.cpp
@@ -497,24 +497,6 @@ void FZenStoreWriter::WriteProjectStoreFile()
 	ProjectStoreWriter.Write(ServiceHostName, ServicePort, IsLocalConnection, *ProjectId, *OplogId, TargetPlatformFName);
 }
 
-// TODO: Move this into FPathViews and change it to take flags instead of bools
-FStringView FPathViewsGetExtension(const FStringView& InPath, bool bIncludeDot, bool bIncludeCompoundExtension)
-{
-	if (!bIncludeCompoundExtension)
-	{
-		return FPathViews::GetExtension(InPath, bIncludeDot);
-	}
-	else
-	{
-		if (const TCHAR* Dot = Algo::Find(FPathViews::GetCleanFilename(InPath), TEXT('.')))
-		{
-			const TCHAR* Extension = bIncludeDot ? Dot : Dot + 1;
-			return FStringView(Extension, UE_PTRDIFF_TO_INT32(InPath.GetData() + InPath.Len() - Extension));
-		}
-		return FStringView();
-	}
-}
-
 void FZenStoreWriter::Initialize(const FCookInfo& Info)
 {
 	CookMode = Info.CookMode;
@@ -633,8 +615,7 @@ void FZenStoreWriter::Initialize(const FCookInfo& Info)
 							// cook and use the ZenCookArtifactReader instead, and allow the cooker to delete files
 							// from the ZenCookArtifactReader when it deletes them from local disk, see UE-351803
 							// and UE-352816.
-							FStringView Extension = FPathViewsGetExtension(ClientPath, false /* bIncludeDot */,
-								true /* bIncludeCompoundExtension */);
+							FStringView Extension = FPathViews::GetExtension(ClientPath, UE::Paths::EFlags::AllowCompoundExtension);
 							if (!CopyFromPreviousCookDenyListExtensions.ContainsByHash(
 								GetTypeHash(Extension), Extension))
 							{
diff --git a/Engine/Source/Runtime/Core/Private/Misc/PathViews.cpp b/Engine/Source/Runtime/Core/Private/Misc/PathViews.cpp
--- a/Engine/Source/Runtime/Core/Private/Misc/PathViews.cpp
+++ b/Engine/Source/Runtime/Core/Private/Misc/PathViews.cpp
@@ -131,16 +131,28 @@ FStringView FPathViews::GetPath(const FStringView& InPath)
 	return FStringView();
 }
 
-FStringView FPathViews::GetExtension(const FStringView& InPath, bool bIncludeDot)
+FStringView FPathViews::GetExtension(const FStringView& InPath, UE::Paths::EFlags Flags)
 {
-	if (const TCHAR* Dot = Algo::FindLast(GetCleanFilename(InPath), TEXT('.')))
+	using namespace UE::Paths;
+
+	FStringView CleanFilename = GetCleanFilename(InPath);
+	const TCHAR* Dot = !EnumHasAnyFlags(Flags, EFlags::AllowCompoundExtension)
+		? Algo::FindLast(CleanFilename, TEXT('.'))
+		: Algo::Find(CleanFilename, TEXT('.'));
+	if (Dot)
 	{
-		const TCHAR* Extension = bIncludeDot ? Dot : Dot + 1;
+		const TCHAR* Extension = EnumHasAnyFlags(Flags, EFlags::IncludeDot) ? Dot : Dot + 1;
 		return FStringView(Extension, UE_PTRDIFF_TO_INT32(InPath.GetData() + InPath.Len() - Extension));
 	}
 	return FStringView();
 }
 
+FStringView FPathViews::GetExtension(const FStringView& InPath, bool bIncludeDot)
+{
+	using namespace UE::Paths;
+	return GetExtension(InPath, bIncludeDot ? EFlags::IncludeDot : EFlags::None);
+}
+
 FStringView FPathViews::GetPathLeaf(const FStringView& InPath)
 {
 	if (const TCHAR* EndPos = Algo::FindLastByPredicate(InPath, UE4PathViews_Private::IsNotSlashOrBackslash))
diff --git a/Engine/Source/Runtime/Core/Private/Misc/Paths.cpp b/Engine/Source/Runtime/Core/Private/Misc/Paths.cpp
--- a/Engine/Source/Runtime/Core/Private/Misc/Paths.cpp
+++ b/Engine/Source/Runtime/Core/Private/Misc/Paths.cpp
@@ -952,18 +952,29 @@ FString FPaths::FindCorrectCase(const FString& Path)
 #endif
 }
 
-FString FPaths::GetExtension( const FString& InPath, bool bIncludeDot )
+FString FPaths::GetExtension(const FString& InPath, UE::Paths::EFlags Flags)
 {
-	const FString Filename = GetCleanFilename(InPath);
-	int32 DotPos = Filename.Find(TEXT("."), ESearchCase::CaseSensitive, ESearchDir::FromEnd);
+	using namespace UE::Paths;
+
+	const FString CleanFilename = GetCleanFilename(InPath);
+	int32 DotPos = !EnumHasAnyFlags(Flags, EFlags::AllowCompoundExtension)
+		? CleanFilename.Find(TEXT("."), ESearchCase::CaseSensitive, ESearchDir::FromEnd)
+		: CleanFilename.Find(TEXT("."), ESearchCase::CaseSensitive, ESearchDir::FromStart);
 	if (DotPos != INDEX_NONE)
 	{
-		return Filename.Mid(DotPos + (bIncludeDot ? 0 : 1));
+		return CleanFilename.Mid(DotPos + (EnumHasAnyFlags(Flags, EFlags::IncludeDot) ? 0 : 1));
 	}
 
 	return TEXT("");
 }
 
+FString FPaths::GetExtension(const FString& InPath, bool bIncludeDot)
+{
+	using namespace UE::Paths;
+
+	return GetExtension(InPath, bIncludeDot ? EFlags::IncludeDot : EFlags::None);
+}
+
 FString FPaths::GetCleanFilename(const FString& InPath)
 {
 	static_assert(INDEX_NONE == -1, "INDEX_NONE assumed to be -1");
diff --git a/Engine/Source/Runtime/Core/Public/Misc/PathTypes.h b/Engine/Source/Runtime/Core/Public/Misc/PathTypes.h
new file mode 100644
--- /dev/null
+++ b/Engine/Source/Runtime/Core/Public/Misc/PathTypes.h
@@ -0,0 +1,31 @@
+// Copyright Epic Games, Inc. All Rights Reserved.
+
+#pragma once
+
+#include "HAL/Platform.h"
+#include "Misc/EnumClassFlags.h"
+
+namespace UE::Paths
+{
+
+/** Flags used for argument passing in FPaths and FPathViews functions. */
+enum class EFlags : uint64
+{
+	None = 0,
+
+	/** In functions that return extensions, request that the extension include the dot. */
+	IncludeDot = 1 << 0,
+
+	/**
+	 * In functions that split extensions out of a path, allow the extension to include multiple dots,
+	 * after the final directory separator. For example:
+	 * "/Path/PackageName.m.uasset" ->
+	 *    (Compound Allowed)      "/Path", "PackageName", ".m.uasset"
+	 *    (Compound not allowed)  "/Path", "PackageName.m", ".uasset"
+	 */
+	AllowCompoundExtension = 1 << 1,
+};
+ENUM_CLASS_FLAGS(EFlags);
+
+} // namespace UE::Paths
+
diff --git a/Engine/Source/Runtime/Core/Public/Misc/PathViews.h b/Engine/Source/Runtime/Core/Public/Misc/PathViews.h
--- a/Engine/Source/Runtime/Core/Public/Misc/PathViews.h
+++ b/Engine/Source/Runtime/Core/Public/Misc/PathViews.h
@@ -4,6 +4,7 @@
 
 #include "Containers/StringFwd.h"
 #include "CoreTypes.h"
+#include "Misc/PathTypes.h"
 
 class FName;
 class FString;
@@ -71,23 +72,34 @@ public:
 	static CORE_API FStringView GetPath(const FStringView& InPath);
 
 	/**
-	 * Returns the portion of the path after the last dot following the last separator, optionally including the dot.
-	 *
+	 * Returns the portion of the path after the last dot following the last separator,
+	 * optionally including the dot, optionally allowing compound extensions.
+	 * 
 	 * Examples: (Using '/' but '\' is valid too.)
-	 * "A/B.C.D" -> "D" (bIncludeDot=false) or ".D" (bIncludeDot=true)
-	 * "A/B/C.D" -> "D" (bIncludeDot=false) or ".D" (bIncludeDot=true)
-	 * "A/B/.D"  -> "D" (bIncludeDot=false) or ".D" (bIncludeDot=true)
-	 * ".D"      -> "D" (bIncludeDot=false) or ".D" (bIncludeDot=true)
+	 * "A/B.C.D" -> "D"    (EFlags::IncludeDot == 0, EFlags::AllowCompoundExtension == 0)
+	 *           -> ".D"   (EFlags::IncludeDot == 1, EFlags::AllowCompoundExtension == 0)
+	 *           -> "C.D"  (EFlags::IncludeDot == 0, EFlags::AllowCompoundExtension == 1)
+	 *           -> ".C.D" (EFlags::IncludeDot == 1, EFlags::AllowCompoundExtension == 1)
+	 * "A/B/C.D" -> "D" (EFlags::IncludeDot == 0) or ".D" (EFlags::IncludeDot == 1)
+	 * "A/B/.D"  -> "D" (EFlags::IncludeDot == 0) or ".D" (EFlags::IncludeDot == 1)
+	 * ".D"      -> "D" (EFlags::IncludeDot == 0) or ".D" (EFlags::IncludeDot == 1)
 	 * "A/B/C"   -> ""
 	 * "A.B/C"   -> ""
 	 * "A.B/"    -> ""
 	 * "A"       -> ""
+	 * "A.B/C.D" -> "D" (EFlags::IncludeDot) or ".D" (EFlags::IncludeDot), even if EFlags::AllowCompoundExtension == 1.
 	 *
-	 * @param bIncludeDot Whether to include the leading dot in the returned view.
+	 * Default: Dot is not included, compound extensions are not allowed.
+	 * @param Flags
+	 *		EFlags::IncludeDot Includes the leading dot in the result.
+	 *		EFlags::AllowCompoundExtension Allows compound extensions (.a.b) in the result.
 	 *
 	 * @return The portion of the path after the last dot following the last separator, optionally including the dot.
 	 */
-	static CORE_API FStringView GetExtension(const FStringView& InPath, bool bIncludeDot=false);
+	static CORE_API FStringView GetExtension(const FStringView& InPath,
+		UE::Paths::EFlags Flags = UE::Paths::EFlags::None);
+	/** Legacy version of GetExtension that uses a bool to specify bIncludeDot. EFlags::AllowCompoundExtension == 0. */
+	static CORE_API FStringView GetExtension(const FStringView& InPath, bool bIncludeDot);
 
 	/**
 	 * Returns the last non-empty path component.
diff --git a/Engine/Source/Runtime/Core/Public/Misc/Paths.h b/Engine/Source/Runtime/Core/Public/Misc/Paths.h
--- a/Engine/Source/Runtime/Core/Public/Misc/Paths.h
+++ b/Engine/Source/Runtime/Core/Public/Misc/Paths.h
@@ -7,6 +7,7 @@
 #include "Containers/UnrealString.h"
 #include "CoreTypes.h"
 #include "HAL/CriticalSection.h"
+#include "Misc/PathTypes.h"
 #include "Templates/UnrealTemplate.h"
 
 #include <type_traits>
@@ -525,12 +526,17 @@ public:
 
 	/**
 	 * Gets the extension for this filename.
-	 *
-	 * @param	bIncludeDot		if true, includes the leading dot in the result
+	 * 
+	 * Default: Dot is not included, compound extensions are not allowed.
+	 * @param Flags
+	 *		EFlags::IncludeDot Includes the leading dot in the result.
+	 *		EFlags::AllowCompoundExtension Allows compound extensions (.a.b) in the result.
 	 *
 	 * @return	the extension of this filename, or an empty string if the filename doesn't have an extension.
 	 */
-	static CORE_API FString GetExtension( const FString& InPath, bool bIncludeDot=false );
+	static CORE_API FString GetExtension(const FString& InPath, UE::Paths::EFlags Flags = UE::Paths::EFlags::None);
+	/** Legacy version of GetExtension that uses a bool to specify bIncludeDot. EFlags::AllowCompoundExtension == 0. */
+	static CORE_API FString GetExtension(const FString& InPath, bool bIncludeDot);
 
 	// Returns the filename (with extension), minus any path information.
 	static CORE_API FString GetCleanFilename(const FString& InPath);
diff --git a/Engine/Source/Runtime/Core/Tests/Misc/PathViewsTest.cpp b/Engine/Source/Runtime/Core/Tests/Misc/PathViewsTest.cpp
--- a/Engine/Source/Runtime/Core/Tests/Misc/PathViewsTest.cpp
+++ b/Engine/Source/Runtime/Core/Tests/Misc/PathViewsTest.cpp
@@ -129,35 +129,60 @@ public:
 
 	bool GetExtensionTest() 
 	{
-		auto RunGetExtensionTest = [this](const TCHAR* InPath, const TCHAR* InExpectedExt, const TCHAR* InExpectedExtDot)
+		using namespace UE::Paths;
+
+		auto RunGetExtTest = [this](const TCHAR* InPath, const TCHAR* InExpectedExt, const TCHAR* InExpectedExtDot)
 		{
 			const FStringView Path = InPath;
 			TestViewTransform([](const FStringView& InPath2) { return FPathViews::GetExtension(InPath2, /*bIncludeDot*/ false); }, Path, InExpectedExt);
+			TestViewTransform([](const FStringView& InPath2) { return FPathViews::GetExtension(InPath2); }, Path, InExpectedExt);
 			TestViewTransform([](const FStringView& InPath2) { return FPathViews::GetExtension(InPath2, /*bIncludeDot*/ true); }, Path, InExpectedExtDot);
+			TestViewTransform([](const FStringView& InPath2) { return FPathViews::GetExtension(InPath2, EFlags::IncludeDot); }, Path, InExpectedExtDot);
 		};
 
-		RunGetExtensionTest(TEXT(""), TEXT(""), TEXT(""));
-		RunGetExtensionTest(TEXT(".txt"), TEXT("txt"), TEXT(".txt"));
-		RunGetExtensionTest(TEXT(".tar.gz"), TEXT("gz"), TEXT(".gz"));
-		RunGetExtensionTest(TEXT(".tar.gz/"), TEXT(""), TEXT(""));
-		RunGetExtensionTest(TEXT(".tar.gz\\"), TEXT(""), TEXT(""));
-		RunGetExtensionTest(TEXT("File"), TEXT(""), TEXT(""));
-		RunGetExtensionTest(TEXT("File.txt"), TEXT("txt"), TEXT(".txt"));
-		RunGetExtensionTest(TEXT("File.tar.gz"), TEXT("gz"), TEXT(".gz"));
-		RunGetExtensionTest(TEXT("File.tar.gz/"), TEXT(""), TEXT(""));
-		RunGetExtensionTest(TEXT("File.tar.gz\\"), TEXT(""), TEXT(""));
-		RunGetExtensionTest(TEXT("C:/Folder/File"), TEXT(""), TEXT(""));
-		RunGetExtensionTest(TEXT("C:\\Folder\\File"), TEXT(""), TEXT(""));
-		RunGetExtensionTest(TEXT("C:/Folder/File.txt"), TEXT("txt"), TEXT(".txt"));
-		RunGetExtensionTest(TEXT("C:\\Folder\\File.txt"), TEXT("txt"), TEXT(".txt"));
-		RunGetExtensionTest(TEXT("C:/Folder/File.tar.gz"), TEXT("gz"), TEXT(".gz"));
-		RunGetExtensionTest(TEXT("C:\\Folder\\File.tar.gz"), TEXT("gz"), TEXT(".gz"));
-		RunGetExtensionTest(TEXT("C:/Folder/First.Last/File"), TEXT(""), TEXT(""));
-		RunGetExtensionTest(TEXT("C:\\Folder\\First.Last\\File"), TEXT(""), TEXT(""));
-		RunGetExtensionTest(TEXT("C:/Folder/First.Last/File.txt"), TEXT("txt"), TEXT(".txt"));
-		RunGetExtensionTest(TEXT("C:\\Folder\\First.Last\\File.txt"), TEXT("txt"), TEXT(".txt"));
-		RunGetExtensionTest(TEXT("C:/Folder/First.Last/File.tar.gz"), TEXT("gz"), TEXT(".gz"));
-		RunGetExtensionTest(TEXT("C:\\Folder\\First.Last\\File.tar.gz"), TEXT("gz"), TEXT(".gz"));
+		auto RunGetExtTestFull = [this](const TCHAR* InPath, const TCHAR* InExpectedExt, const TCHAR* InExpectedExtDot,
+			const TCHAR* InExpectedExtCompound, const TCHAR* InExpectedExtDotCompound)
+		{
+			const FStringView Path = InPath;
+			TestViewTransform([](const FStringView& InPath2)
+				{ return FPathViews::GetExtension(InPath2); },
+				Path, InExpectedExt);
+			TestViewTransform([](const FStringView& InPath2)
+				{ return FPathViews::GetExtension(InPath2, EFlags::None); },
+				Path, InExpectedExt);
+			TestViewTransform([](const FStringView& InPath2)
+				{ return FPathViews::GetExtension(InPath2, EFlags::IncludeDot); },
+				Path, InExpectedExtDot);
+			TestViewTransform([](const FStringView& InPath2)
+				{ return FPathViews::GetExtension(InPath2, EFlags::AllowCompoundExtension); },
+				Path, InExpectedExtCompound);
+			TestViewTransform([](const FStringView& InPath2)
+				{ return FPathViews::GetExtension(InPath2, EFlags::IncludeDot | EFlags::AllowCompoundExtension); },
+				Path, InExpectedExtDotCompound);
+			};
+
+		RunGetExtTest(TEXT(""),											TEXT(""),	TEXT(""));
+		RunGetExtTest(TEXT(".txt"), 									TEXT("txt"),TEXT(".txt"));
+		RunGetExtTest(TEXT(".tar.gz"),									TEXT("gz"), TEXT(".gz"));
+		RunGetExtTest(TEXT(".tar.gz/"),									TEXT(""),	TEXT(""));
+		RunGetExtTest(TEXT(".tar.gz\\"),								TEXT(""),	TEXT(""));
+		RunGetExtTest(TEXT("File"),										TEXT(""),	TEXT(""));
+		RunGetExtTest(TEXT("File.txt"),									TEXT("txt"),TEXT(".txt"));
+		RunGetExtTestFull(TEXT("File.tar.gz"),							TEXT("gz"), TEXT(".gz"),	TEXT("tar.gz"), TEXT(".tar.gz"));
+		RunGetExtTestFull(TEXT("File.tar.gz/"),							TEXT(""),	TEXT(""),		TEXT(""),		TEXT(""));
+		RunGetExtTestFull(TEXT("File.tar.gz\\"),						TEXT(""),	TEXT(""),		TEXT(""),		TEXT(""));
+		RunGetExtTest(TEXT("C:/Folder/File"),							TEXT(""),	TEXT(""));
+		RunGetExtTest(TEXT("C:\\Folder\\File"),							TEXT(""),	TEXT(""));
+		RunGetExtTest(TEXT("C:/Folder/File.txt"),						TEXT("txt"),TEXT(".txt"));
+		RunGetExtTest(TEXT("C:\\Folder\\File.txt"),						TEXT("txt"),TEXT(".txt"));
+		RunGetExtTestFull(TEXT("C:/Folder/File.tar.gz"),				TEXT("gz"), TEXT(".gz"),	TEXT("tar.gz"),	TEXT(".tar.gz"));
+		RunGetExtTestFull(TEXT("C:\\Folder\\File.tar.gz"),				TEXT("gz"), TEXT(".gz"),	TEXT("tar.gz"),	TEXT(".tar.gz"));
+		RunGetExtTestFull(TEXT("C:/Folder/First.Last/File"),			TEXT(""),	TEXT(""),		TEXT(""),		TEXT(""));
+		RunGetExtTestFull(TEXT("C:\\Folder\\First.Last\\File"),			TEXT(""),	TEXT(""),		TEXT(""),		TEXT(""));
+		RunGetExtTestFull(TEXT("C:/Folder/First.Last/File.txt"),		TEXT("txt"),TEXT(".txt"),	TEXT("txt"),	TEXT(".txt"));
+		RunGetExtTestFull(TEXT("C:\\Folder\\First.Last\\File.txt"),		TEXT("txt"),TEXT(".txt"),	TEXT("txt"),	TEXT(".txt"));
+		RunGetExtTestFull(TEXT("C:/Folder/First.Last/File.tar.gz"),		TEXT("gz"), TEXT(".gz"),	TEXT("tar.gz"), TEXT(".tar.gz"));
+		RunGetExtTestFull(TEXT("C:\\Folder\\First.Last\\File.tar.gz"),	TEXT("gz"), TEXT(".gz"),	TEXT("tar.gz"), TEXT(".tar.gz"));
 
 		return true;
 	}
diff --git a/Engine/Source/Runtime/Core/Tests/Misc/PathsTest.cpp b/Engine/Source/Runtime/Core/Tests/Misc/PathsTest.cpp
--- a/Engine/Source/Runtime/Core/Tests/Misc/PathsTest.cpp
+++ b/Engine/Source/Runtime/Core/Tests/Misc/PathsTest.cpp
@@ -38,30 +38,87 @@ TConstArrayView<PathTest::FTestPair> PathTest::ExpectedRelativeToAbsolutePaths =
 
 TEST_CASE_NAMED(FPathTests, "System::Core::Misc::Paths", "[ApplicationContextMask][SmokeFilter]")
 {
+	using namespace UE::Paths;
+
 	TestCollapseRelativeDirectories<FPaths, FString>();
 	TestCreateStandardFilename();
 
 	// Extension texts
 	{
-		auto RunGetExtensionTest = [](const TCHAR* InPath, const TCHAR* InExpectedExt)
-		{
-			// Run test
-			const FString Ext = FPaths::GetExtension(FString(InPath));
-			if (Ext != InExpectedExt)
+		auto RunGetTest = [](const TCHAR* InPath, const TCHAR* InExpectedExt, const TCHAR* InExpectedExtDot)
 			{
-				FAIL_CHECK(FString::Printf(TEXT("Path '%s' failed to get the extension (got '%s', expected '%s')."), InPath, *Ext, InExpectedExt));
-			}
-		};
+				FString Ext = FPaths::GetExtension(FString(InPath));
+				if (Ext != InExpectedExt)
+				{
+					FAIL_CHECK(FString::Printf(TEXT("Path '%s' failed to get the extension (got '%s', expected '%s')."),
+						InPath, *Ext, InExpectedExt));
+				}
+				Ext = FPaths::GetExtension(FString(InPath), false /* bIncludeDot */);
+				if (Ext != InExpectedExt)
+				{
+					FAIL_CHECK(FString::Printf(TEXT("Path '%s' failed to get the extension (got '%s', expected '%s')."),
+						InPath, *Ext, InExpectedExt));
+				}
+				Ext = FPaths::GetExtension(FString(InPath), true/* bIncludeDot */);
+				if (Ext != InExpectedExtDot)
+				{
+					FAIL_CHECK(FString::Printf(TEXT("Path '%s' failed to get the extension (got '%s', expected '%s')."),
+						InPath, *Ext, InExpectedExtDot));
+				}
+				Ext = FPaths::GetExtension(FString(InPath), EFlags::IncludeDot);
+				if (Ext != InExpectedExtDot)
+				{
+					FAIL_CHECK(FString::Printf(TEXT("Path '%s' failed to get the extension (got '%s', expected '%s')."),
+						InPath, *Ext, InExpectedExtDot));
+				}
+			};
+
+		auto RunGetTestFull = [](const TCHAR* InPath, const TCHAR* InExpectedExt, const TCHAR* InExpectedExtDot,
+			const TCHAR* InExpectedExtCompound, const TCHAR* InExpectedExtDotCompound)
+			{
+				FString Ext = FPaths::GetExtension(FString(InPath));
+				if (Ext != InExpectedExt)
+				{
+					FAIL_CHECK(FString::Printf(TEXT("Path '%s' failed to get the extension (got '%s', expected '%s')."),
+						InPath, *Ext, InExpectedExt));
+				}
+				Ext = FPaths::GetExtension(FString(InPath), EFlags::None);
+				if (Ext != InExpectedExt)
+				{
+					FAIL_CHECK(FString::Printf(TEXT("Path '%s' failed to get the extension (got '%s', expected '%s')."),
+						InPath, *Ext, InExpectedExt));
+				}
+				Ext = FPaths::GetExtension(FString(InPath), EFlags::IncludeDot);
+				if (Ext != InExpectedExtDot)
+				{
+					FAIL_CHECK(FString::Printf(TEXT("Path '%s' failed to get the extension (got '%s', expected '%s')."),
+						InPath, *Ext, InExpectedExtDot));
+				}
+				Ext = FPaths::GetExtension(FString(InPath), EFlags::AllowCompoundExtension);
+				if (Ext != InExpectedExtCompound)
+				{
+					FAIL_CHECK(FString::Printf(TEXT("Path '%s' failed to get the extension (got '%s', expected '%s')."),
+						InPath, *Ext, InExpectedExtCompound));
+				}
+				Ext = FPaths::GetExtension(FString(InPath), EFlags::IncludeDot | EFlags::AllowCompoundExtension);
+				if (Ext != InExpectedExtDotCompound)
+				{
+					FAIL_CHECK(FString::Printf(TEXT("Path '%s' failed to get the extension (got '%s', expected '%s')."),
+						InPath, *Ext, InExpectedExtDotCompound));
+				}
+			};
+
+
 
-		RunGetExtensionTest(TEXT("file"),									TEXT(""));
-		RunGetExtensionTest(TEXT("file.txt"),								TEXT("txt"));
-		RunGetExtensionTest(TEXT("file.tar.gz"),							TEXT("gz"));
-		RunGetExtensionTest(TEXT("C:/Folder/file"),							TEXT(""));
-		RunGetExtensionTest(TEXT("C:/Folder/file.txt"),						TEXT("txt"));
-		RunGetExtensionTest(TEXT("C:/Folder/file.tar.gz"),					TEXT("gz"));
-		RunGetExtensionTest(TEXT("C:/Folder/First.Last/file"),				TEXT(""));
-		RunGetExtensionTest(TEXT("C:/Folder/First.Last/file.txt"),			TEXT("txt"));
-		RunGetExtensionTest(TEXT("C:/Folder/First.Last/file.tar.gz"),		TEXT("gz"));
+		RunGetTest(TEXT("file"),								TEXT(""),	TEXT(""));
+		RunGetTest(TEXT("file.txt"),							TEXT("txt"),TEXT(".txt"));
+		RunGetTest(TEXT("file.tar.gz"),							TEXT("gz"),	TEXT(".gz"));
+		RunGetTest(TEXT("C:/Folder/file"),						TEXT(""),	TEXT(""));
+		RunGetTest(TEXT("C:/Folder/file.txt"),					TEXT("txt"),TEXT(".txt"));
+		RunGetTestFull(TEXT("C:/Folder/file.tar.gz"),			TEXT("gz"),	TEXT(".gz"),	TEXT("tar.gz"), TEXT(".tar.gz"));
+		RunGetTestFull(TEXT("C:/Folder/First.Last/file"),		TEXT(""),	TEXT(""),		TEXT(""),		TEXT(""));
+		RunGetTestFull(TEXT("C:/Folder/First.Last/file.txt"),	TEXT("txt"),TEXT(".txt"),	TEXT("txt"),	TEXT(".txt"));
+		RunGetTestFull(TEXT("C:/Folder/First.Last/file.tar.gz"),TEXT("gz"),	TEXT(".gz"),	TEXT("tar.gz"), TEXT(".tar.gz"));
 
 		auto RunSetExtensionTest = [](const TCHAR* InPath, const TCHAR* InNewExt, const FString& InExpectedPath)
 		{
