Enabled memory optimizing options for mimalloc when AGGRESSIVE_MEMORY_SAVING is set to 1

#Synced-CL 49977366

diff --git a/Engine/Source/Runtime/Core/Private/HAL/MallocMimalloc.cpp b/Engine/Source/Runtime/Core/Private/HAL/MallocMimalloc.cpp
--- a/Engine/Source/Runtime/Core/Private/HAL/MallocMimalloc.cpp
+++ b/Engine/Source/Runtime/Core/Private/HAL/MallocMimalloc.cpp
@@ -49,6 +49,24 @@ FMallocMimalloc::FMallocMimalloc()
 	FPlatformMemory::MiMallocInit();
 
 	ApplyMiMallocMemoryResetDelayChange();
+
+#if AGGRESSIVE_MEMORY_SAVING
+#	if MI_MALLOC_VERSION == 212
+		mi_option_set(mi_option_deprecated_page_reset, 1);
+		mi_option_set(mi_option_abandoned_page_purge, 1);
+		mi_option_set(mi_option_deprecated_segment_reset, 1);
+		mi_option_set(mi_option_purge_delay, 0);
+#	else
+		mi_option_set(mi_option_reset_decommits, 1);
+		mi_option_set(mi_option_page_reset, 1);
+		mi_option_set(mi_option_abandoned_page_reset, 1);
+		mi_option_set(mi_option_segment_reset, 1);
+		mi_option_set(mi_option_reset_delay, 0);
+		mi_option_set(mi_option_segment_decommit_delay, 0);
+#	endif
+
+	mi_option_set(mi_option_eager_commit_delay, 4);
+#endif
 }
 
 void* FMallocMimalloc::TryMalloc( SIZE_T Size, uint32 Alignment )
