[MemoryLayout]
- Fix on-demand initialization of layout classes that weren't thread-safe

#Synced-CL 50201644

diff --git a/Engine/Source/Runtime/Core/Public/Serialization/MemoryLayout.h b/Engine/Source/Runtime/Core/Public/Serialization/MemoryLayout.h
--- a/Engine/Source/Runtime/Core/Public/Serialization/MemoryLayout.h
+++ b/Engine/Source/Runtime/Core/Public/Serialization/MemoryLayout.h
@@ -142,6 +142,7 @@ struct FTypeLayoutDesc
 	uint8 NumBases;
 	uint8 NumVirtualBases;
 	uint8 IsIntrinsic : 1;
+	UE_DEPRECATED(5.8, "IsInitialized is not used anymore")
 	uint8 IsInitialized : 1;
 
 	friend inline bool operator==(const FTypeLayoutDesc& Lhs, const FTypeLayoutDesc& Rhs) { return &Lhs == &Rhs; }
@@ -517,10 +518,10 @@ PRAGMA_ENABLE_DEPRECATION_WARNINGS
 	} \
 	public: static FTypeLayoutDesc& StaticGetTypeLayout() { \
 		static_assert(TValidateInterfaceHelper<UE_REMOVE_OPTIONAL_PARENS(T), ETypeLayoutInterface::InInterface>::Value, #InInterface " is invalid interface for " #T); \
-		alignas(FTypeLayoutDesc) static uint8 TypeBuffer[sizeof(FTypeLayoutDesc)] = { 0 }; \
-		FTypeLayoutDesc& TypeDesc = *(FTypeLayoutDesc*)TypeBuffer; \
-		if (!TypeDesc.IsInitialized) { \
-			TypeDesc.IsInitialized = true; \
+		static FTypeLayoutDesc& StaticTypeDesc = []() -> FTypeLayoutDesc& \
+		{ \
+			alignas(FTypeLayoutDesc) static uint8 TypeBuffer[sizeof(FTypeLayoutDesc)] = { 0 }; \
+			FTypeLayoutDesc& TypeDesc = *(FTypeLayoutDesc*)TypeBuffer; \
 			TypeDesc.Name = TEXT(#T); \
 			TypeDesc.WriteFrozenMemoryImageFunc = TGetFreezeImageHelper<UE_REMOVE_OPTIONAL_PARENS(T)>::Do(); \
 			TypeDesc.UnfrozenCopyFunc = &Freeze::DefaultUnfrozenCopy; \
@@ -536,8 +537,9 @@ PRAGMA_ENABLE_DEPRECATION_WARNINGS
 			InternalLinkType<1>::Initialize(TypeDesc); \
 			InternalInitializeBases<UE_REMOVE_OPTIONAL_PARENS(T)>(TypeDesc); \
 			FTypeLayoutDesc::Initialize(TypeDesc); \
-		} \
-		return TypeDesc; } \
+			return TypeDesc; \
+		}(); \
+		return StaticTypeDesc; } \
 	public: INTERNAL_LAYOUT_INTERFACE_PREFIX(InInterface)(UE_EMPTY) const FTypeLayoutDesc& GetTypeLayout() const INTERNAL_LAYOUT_INTERFACE_INLINE_IMPL(InInterface) \
 	INTERNAL_DECLARE_TYPE_LAYOUT_COMMON(T, InInterface)
 
@@ -568,10 +570,10 @@ PRAGMA_ENABLE_DEPRECATION_WARNINGS
 	} \
 	UE_REMOVE_OPTIONAL_PARENS(TemplatePrefix) FTypeLayoutDesc& UE_REMOVE_OPTIONAL_PARENS(T)::StaticGetTypeLayout() { \
 		static_assert(TValidateInterfaceHelper<UE_REMOVE_OPTIONAL_PARENS(T), InterfaceType>::Value, "Invalid interface for " #T); \
-		alignas(FTypeLayoutDesc) static uint8 TypeBuffer[sizeof(FTypeLayoutDesc)] = { 0 }; \
-		FTypeLayoutDesc& TypeDesc = *(FTypeLayoutDesc*)TypeBuffer; \
-		if (!TypeDesc.IsInitialized) { \
-			TypeDesc.IsInitialized = true; \
+		static FTypeLayoutDesc& StaticTypeDesc = []() -> FTypeLayoutDesc& \
+		{ \
+			alignas(FTypeLayoutDesc) static uint8 TypeBuffer[sizeof(FTypeLayoutDesc)] = { 0 }; \
+			FTypeLayoutDesc& TypeDesc = *(FTypeLayoutDesc*)TypeBuffer; \
 			TypeDesc.Name = TEXT(#T); \
 			TypeDesc.WriteFrozenMemoryImageFunc = TGetFreezeImageHelper<UE_REMOVE_OPTIONAL_PARENS(T)>::Do(); \
 			TypeDesc.UnfrozenCopyFunc = &Freeze::DefaultUnfrozenCopy; \
@@ -587,8 +589,9 @@ PRAGMA_ENABLE_DEPRECATION_WARNINGS
 			InternalLinkType<1>::Initialize(TypeDesc); \
 			InternalInitializeBases<UE_REMOVE_OPTIONAL_PARENS(T)>(TypeDesc); \
 			FTypeLayoutDesc::Initialize(TypeDesc); \
-		} \
-		return TypeDesc; }
+			return TypeDesc; \
+		}(); \
+		return StaticTypeDesc; }
 
 //#define INTERNAL_REGISTER_TYPE_LAYOUT(T) static const FRegisterTypeLayoutDesc ANONYMOUS_VARIABLE(RegisterTypeLayoutDesc)(UE_REMOVE_OPTIONAL_PARENS(T)::StaticGetTypeLayout())
 #define INTERNAL_REGISTER_TYPE_LAYOUT(T) static const FDelayedAutoRegisterHelper ANONYMOUS_VARIABLE(DelayedAutoRegisterHelper)(EDelayedRegisterRunPhase::ShaderTypesReady, []{ FTypeLayoutDesc::Register(T::StaticGetTypeLayout()); });
@@ -681,10 +684,10 @@ inline void DeleteObjectFromLayout(T* Object, const FPointerTableBase* PtrTable
 			Freeze::DestroyObject(static_cast<UE_REMOVE_OPTIONAL_PARENS(T)*>(Object), PtrTable, bIsFrozen); \
 		} \
 		static const FTypeLayoutDesc& Do() { \
-			alignas(FTypeLayoutDesc) static uint8 TypeBuffer[sizeof(FTypeLayoutDesc)] = { 0 }; \
-			FTypeLayoutDesc& TypeDesc = *(FTypeLayoutDesc*)TypeBuffer; \
-			if (!TypeDesc.IsInitialized) { \
-				TypeDesc.IsInitialized = true; \
+			static FTypeLayoutDesc& StaticTypeDesc = []() -> FTypeLayoutDesc& \
+			{ \
+				alignas(FTypeLayoutDesc) static uint8 TypeBuffer[sizeof(FTypeLayoutDesc)] = { 0 }; \
+				FTypeLayoutDesc& TypeDesc = *(FTypeLayoutDesc*)TypeBuffer; \
 				TypeDesc.IsIntrinsic = true; \
 				TypeDesc.Name = TEXT(#T); \
 				TypeDesc.WriteFrozenMemoryImageFunc = &CallWriteMemoryImage; \
@@ -697,8 +700,9 @@ inline void DeleteObjectFromLayout(T* Object, const FPointerTableBase* PtrTable
 				TypeDesc.Alignment = alignof(UE_REMOVE_OPTIONAL_PARENS(T)); \
 				TypeDesc.Interface = ETypeLayoutInterface::NonVirtual; \
 				TypeDesc.SizeFromFields = sizeof(UE_REMOVE_OPTIONAL_PARENS(T)); \
-			} \
-			return TypeDesc; } }; \
+				return TypeDesc; \
+			}(); \
+			return StaticTypeDesc; } }; \
 	UE_REMOVE_OPTIONAL_PARENS(TemplatePrefix) struct TGetTypeLayoutHelper<UE_REMOVE_OPTIONAL_PARENS(T)> { \
 		UE_STATIC_ONLY(TGetTypeLayoutHelper); \
 		static const FTypeLayoutDesc& Do(const UE_REMOVE_OPTIONAL_PARENS(T)&) { return TStaticGetTypeLayoutHelper<UE_REMOVE_OPTIONAL_PARENS(T)>::Do(); }}
@@ -738,10 +742,10 @@ inline void DeleteObjectFromLayout(T* Object, const FPointerTableBase* PtrTable
 			Freeze::DestroyObject(static_cast<UE_REMOVE_OPTIONAL_PARENS(T)*>(Object), PtrTable, bIsFrozen); \
 		} \
 		const FTypeLayoutDesc& TStaticGetTypeLayoutHelper<UE_REMOVE_OPTIONAL_PARENS(T)>::Do() { \
-			alignas(FTypeLayoutDesc) static uint8 TypeBuffer[sizeof(FTypeLayoutDesc)] = { 0 }; \
-			FTypeLayoutDesc& TypeDesc = *(FTypeLayoutDesc*)TypeBuffer; \
-			if (!TypeDesc.IsInitialized) { \
-				TypeDesc.IsInitialized = true; \
+			static FTypeLayoutDesc& StaticTypeDesc = []() -> FTypeLayoutDesc& \
+			{ \
+				alignas(FTypeLayoutDesc) static uint8 TypeBuffer[sizeof(FTypeLayoutDesc)] = { 0 }; \
+				FTypeLayoutDesc& TypeDesc = *(FTypeLayoutDesc*)TypeBuffer; \
 				TypeDesc.IsIntrinsic = true; \
 				TypeDesc.Name = TEXT(#T); \
 				TypeDesc.WriteFrozenMemoryImageFunc = &CallWriteMemoryImage; \
@@ -754,8 +758,9 @@ inline void DeleteObjectFromLayout(T* Object, const FPointerTableBase* PtrTable
 				TypeDesc.Alignment = alignof(UE_REMOVE_OPTIONAL_PARENS(T)); \
 				TypeDesc.Interface = ETypeLayoutInterface::NonVirtual; \
 				TypeDesc.SizeFromFields = sizeof(UE_REMOVE_OPTIONAL_PARENS(T)); \
-			} \
-			return TypeDesc; }
+				return TypeDesc; \
+			}(); \
+			return StaticTypeDesc; }
 
 #define DECLARE_INTRINSIC_TYPE_LAYOUT(T) DECLARE_TEMPLATE_INTRINSIC_TYPE_LAYOUT(template<>, T)
 
