Added support for mimalloc v2.1.2 on iOS, enabled via UE_APPLE_USE_MIMALLOC_MIN_FREE_RAM_MB define

#Synced-CL 49977366

diff --git a/Engine/Source/Runtime/Core/Core.Build.cs b/Engine/Source/Runtime/Core/Core.Build.cs
--- a/Engine/Source/Runtime/Core/Core.Build.cs
+++ b/Engine/Source/Runtime/Core/Core.Build.cs
@@ -221,6 +221,9 @@ public class Core : ModuleRules
 			PrivateDefinitions.Add("UE_MB3_MAX_SMALL_POOL_SIZE=14336");
 			PrivateDefinitions.Add("UE_MBC_MAX_LISTED_SMALL_POOL_SIZE=14336");
 			PrivateDefinitions.Add("UE_MBC_NUM_LISTED_SMALL_POOLS=48");
+
+			PrivateDependencyModuleNames.Add("mimalloc212");
+			PrivateDefinitions.Add("PLATFORM_BUILDS_MIMALLOC=1");
 		}
 
 		if (Target.bCompileICU == true)
diff --git a/Engine/Source/Runtime/Core/Private/Apple/ApplePlatformMemory.cpp b/Engine/Source/Runtime/Core/Private/Apple/ApplePlatformMemory.cpp
--- a/Engine/Source/Runtime/Core/Private/Apple/ApplePlatformMemory.cpp
+++ b/Engine/Source/Runtime/Core/Private/Apple/ApplePlatformMemory.cpp
@@ -27,6 +27,7 @@
 #include "HAL/MallocBinned.h"
 #include "HAL/MallocBinned2.h"
 #include "HAL/MallocBinned3.h"
+#include "HAL/MallocMimalloc.h"
 #include "CoreGlobals.h"
 
 #include <stdlib.h>
@@ -221,45 +222,54 @@ void FApplePlatformMemory::Init()
 	
 }
 
-void FApplePlatformMemory::SetAllocatorToUse()
+#ifndef UE_APPLE_USE_MIMALLOC_MIN_FREE_RAM_MB
+#	define UE_APPLE_USE_MIMALLOC_MIN_FREE_RAM_MB 0
+#endif
+
+void SelectAllocatorToUse(const FPlatformMemoryStats& MemStats)
 {
     // force Ansi allocator in particular cases
     if(getenv("UE4_FORCE_MALLOC_ANSI") != nullptr)
     {
 		NSLog(@"UE4_FORCE_MALLOC_ANSI is set, using Ansi allocator.\n");
-        AllocatorToUse = EMemoryAllocatorToUse::Ansi;
+        FPlatformMemory::AllocatorToUse = FPlatformMemory::EMemoryAllocatorToUse::Ansi;
         return;
     }
     if (FORCE_ANSI_ALLOCATOR)
     {
 		NSLog(@"FORCE_ANSI_ALLOCATOR defined, using Ansi allocator.\n");
-        AllocatorToUse = EMemoryAllocatorToUse::Ansi;
+		FPlatformMemory::AllocatorToUse = FPlatformMemory::EMemoryAllocatorToUse::Ansi;
         return;
     }
-
+	if (UE_APPLE_USE_MIMALLOC_MIN_FREE_RAM_MB && ((MemStats.AvailablePhysical / 1024 / 1024) >= UE_APPLE_USE_MIMALLOC_MIN_FREE_RAM_MB))
+	{
+		NSLog(@"Using mimalloc allocator.\n");
+		FPlatformMemory::AllocatorToUse = FPlatformMemory::EMemoryAllocatorToUse::Mimalloc;
+		return;
+	}
 	if (USE_MALLOC_BINNED3)
 	{
-		if (!CanOverallocateVirtualMemory())
+		if (!FApplePlatformMemory::CanOverallocateVirtualMemory())
 		{
 			NSLog(@"MallocBinned3 requested but com.apple.developer.kernel.extended-virtual-addressing entitlement not found. Check your entitlements. Falling back to Ansi.\n");
-			AllocatorToUse = EMemoryAllocatorToUse::Ansi;
+			FPlatformMemory::AllocatorToUse = FPlatformMemory::EMemoryAllocatorToUse::Ansi;
 			return;
 		}
 
 		NSLog(@"Using MallocBinned3 allocator.\n");
-		AllocatorToUse = EMemoryAllocatorToUse::Binned3;
+		FPlatformMemory::AllocatorToUse = FPlatformMemory::EMemoryAllocatorToUse::Binned3;
 		return;
 	}
     else if (USE_MALLOC_BINNED2)
     {
 		NSLog(@"Using MallocBinned2 allocator.\n");
-        AllocatorToUse = EMemoryAllocatorToUse::Binned2;
+		FPlatformMemory::AllocatorToUse = FPlatformMemory::EMemoryAllocatorToUse::Binned2;
         return;
     }
     else
     {
 		NSLog(@"Defaulting to Ansi allocator.\n");
-        AllocatorToUse = EMemoryAllocatorToUse::Ansi;
+		FPlatformMemory::AllocatorToUse = FPlatformMemory::EMemoryAllocatorToUse::Ansi;
         return;
     }
 }
@@ -272,13 +282,13 @@ FMalloc* FApplePlatformMemory::BaseAllocator()
 		return Instance;
 	}
 
-	FPlatformMemoryStats MemStats = FApplePlatformMemory::GetStats();
+	const FPlatformMemoryStats MemStats = FApplePlatformMemory::GetStats();
 #if ENABLE_LOW_LEVEL_MEM_TRACKER
 	FLowLevelMemTracker::Get().SetProgramSize(MemStats.UsedPhysical);
 #endif
 	FPlatformMemory::ProgramSize = MemStats.UsedPhysical;
 
-    SetAllocatorToUse();
+	SelectAllocatorToUse(MemStats);
     
     switch (AllocatorToUse)
     {
@@ -288,6 +298,12 @@ FMalloc* FApplePlatformMemory::BaseAllocator()
             break;
         }
 
+		case EMemoryAllocatorToUse::Mimalloc:
+		{
+			Instance = new FMallocMimalloc();
+			break;
+		}
+
 		case EMemoryAllocatorToUse::Binned3:
 		{
 			Instance = new FMallocBinned3();
diff --git a/Engine/Source/Runtime/Core/Public/Apple/ApplePlatform.h b/Engine/Source/Runtime/Core/Public/Apple/ApplePlatform.h
--- a/Engine/Source/Runtime/Core/Public/Apple/ApplePlatform.h
+++ b/Engine/Source/Runtime/Core/Public/Apple/ApplePlatform.h
@@ -27,6 +27,7 @@
 #define PLATFORM_SUPPORTS_TEXTURE_STREAMING				1
 #define PLATFORM_SUPPORTS_STACK_SYMBOLS					1
 #define PLATFORM_IS_ANSI_MALLOC_THREADSAFE				1
+#define PLATFORM_SUPPORTS_MIMALLOC						1
 
 #define PLATFORM_BREAK()                                __builtin_debugtrap()
 
diff --git a/Engine/Source/Runtime/Core/Public/Apple/ApplePlatformMemory.h b/Engine/Source/Runtime/Core/Public/Apple/ApplePlatformMemory.h
--- a/Engine/Source/Runtime/Core/Public/Apple/ApplePlatformMemory.h
+++ b/Engine/Source/Runtime/Core/Public/Apple/ApplePlatformMemory.h
@@ -67,7 +67,8 @@ struct CORE_API FApplePlatformMemory : public FGenericPlatformMemory
 	static bool PtrIsFromNanoMalloc( void* Ptr);
 	static bool IsNanoMallocAvailable();
 	static void NanoMallocInit();
-    static void SetAllocatorToUse();
+	UE_DEPRECATED(5.8, "This doesn't belong in a public interface")
+	inline void SetAllocatorToUse() {}
 
 	class FPlatformVirtualMemoryBlock : public FBasicVirtualMemoryBlock
 	{
diff --git a/Engine/Source/Runtime/Core/Public/Mac/MacPlatform.h b/Engine/Source/Runtime/Core/Public/Mac/MacPlatform.h
--- a/Engine/Source/Runtime/Core/Public/Mac/MacPlatform.h
+++ b/Engine/Source/Runtime/Core/Public/Mac/MacPlatform.h
@@ -49,7 +49,6 @@ typedef FMacPlatformTypes FPlatformTypes;
 //#define PLATFORM_USE_PTHREADS							1
 #define PLATFORM_MAX_FILEPATH_LENGTH_DEPRECATED			MAC_MAX_PATH
 #define PLATFORM_SUPPORTS_TBB							1
-#define PLATFORM_SUPPORTS_MIMALLOC						1
 #define PLATFORM_SUPPORTS_MESH_SHADERS					1
 #define PLATFORM_SUPPORTS_BINDLESS_RENDERING			1
 #define PLATFORM_SUPPORTS_GEOMETRY_SHADERS				1
