Enable mimalloc as a default allocator for Mac Editor

#Synced-CL 49977366

diff --git a/Engine/Source/Runtime/Core/Private/Mac/MacPlatformMemory.cpp b/Engine/Source/Runtime/Core/Private/Mac/MacPlatformMemory.cpp
--- a/Engine/Source/Runtime/Core/Private/Mac/MacPlatformMemory.cpp
+++ b/Engine/Source/Runtime/Core/Private/Mac/MacPlatformMemory.cpp
@@ -86,10 +86,15 @@ FMalloc* FMacPlatformMemory::BaseAllocator()
 		return Instance;
 	}
 
-	if (FORCE_ANSI_ALLOCATOR || IS_PROGRAM)
+	if (FORCE_ANSI_ALLOCATOR)
 	{
 		AllocatorToUse = EMemoryAllocatorToUse::Ansi;
 	}
+	else if ((WITH_EDITORONLY_DATA || IS_PROGRAM || USING_INSTRUMENTATION) && MIMALLOC_ENABLED)
+	{
+		// Make mimalloc a default allocator on Editor as it has a significantly better perf
+		AllocatorToUse = EMemoryAllocatorToUse::Mimalloc;
+	}
 	else if ((WITH_EDITORONLY_DATA || IS_PROGRAM) && TBBMALLOC_ENABLED)
 	{
 		AllocatorToUse = EMemoryAllocatorToUse::TBB;
