Added TargetRules.bEnableConfigSystem to allow disabling loading config

Disabled ICU config warnings when ICU is disabled.

#Synced-CL 50253651

diff --git a/Engine/Source/Programs/UnrealBuildTool/Configuration/Rules/ReadOnlyTargetRules.cs b/Engine/Source/Programs/UnrealBuildTool/Configuration/Rules/ReadOnlyTargetRules.cs
--- a/Engine/Source/Programs/UnrealBuildTool/Configuration/Rules/ReadOnlyTargetRules.cs
+++ b/Engine/Source/Programs/UnrealBuildTool/Configuration/Rules/ReadOnlyTargetRules.cs
@@ -168,6 +168,8 @@ namespace UnrealBuildTool
 
 		public bool bCompilePython => Inner.bCompilePython;
 
+		public bool bEnableConfigSystem => Inner.bEnableConfigSystem;
+
 		public bool bBuildEditor => Inner.bBuildEditor;
 
 		public bool bBuildRequiresCookedData => Inner.bBuildRequiresCookedData;
diff --git a/Engine/Source/Programs/UnrealBuildTool/Configuration/Rules/TargetRules.cs b/Engine/Source/Programs/UnrealBuildTool/Configuration/Rules/TargetRules.cs
--- a/Engine/Source/Programs/UnrealBuildTool/Configuration/Rules/TargetRules.cs
+++ b/Engine/Source/Programs/UnrealBuildTool/Configuration/Rules/TargetRules.cs
@@ -1095,6 +1095,12 @@ namespace UnrealBuildTool
 		[RequiresUniqueBuildEnvironment]
 		public bool bCompilePython { get; set; } = true;
 
+		/// <summary>
+		/// Whether to enable the config system (GConfig)
+		/// </summary>
+		[RequiresUniqueBuildEnvironment]
+		public bool bEnableConfigSystem { get; set; } = true;
+
 		/// <summary>
 		/// Whether to compile with support for multiple platform-native views per slate window (right now only implemented on macOS)
 		/// </summary>
diff --git a/Engine/Source/Runtime/Core/Core.Build.cs b/Engine/Source/Runtime/Core/Core.Build.cs
--- a/Engine/Source/Runtime/Core/Core.Build.cs
+++ b/Engine/Source/Runtime/Core/Core.Build.cs
@@ -232,6 +232,8 @@ public class Core : ModuleRules
 		}
 		PublicDefinitions.Add("UE_ENABLE_ICU=" + (Target.bCompileICU ? "1" : "0")); // Enable/disable (=1/=0) ICU usage in the codebase. NOTE: This flag is for use while integrating ICU and will be removed afterward.
 
+		PublicDefinitions.Add($"UE_ENABLE_CONFIG_SYSTEM={(Target.bEnableConfigSystem ? 1 : 0)}");
+
 		// If we're compiling with the engine, then add Core's engine dependencies
 		if (Target.bCompileAgainstEngine && !Target.bBuildRequiresCookedData)
 		{
diff --git a/Engine/Source/Runtime/Core/Private/Misc/ConfigCacheIni.cpp b/Engine/Source/Runtime/Core/Private/Misc/ConfigCacheIni.cpp
--- a/Engine/Source/Runtime/Core/Private/Misc/ConfigCacheIni.cpp
+++ b/Engine/Source/Runtime/Core/Private/Misc/ConfigCacheIni.cpp
@@ -39,9 +39,9 @@
 #include "Logging/MessageLog.h"
 #include <limits>
 
-// if set to 1, this will make sure FCOnfigCacheIni::GetString (and friends) is thread safe from dynamic layer loading, and SafeReload
+// if set to 1, this will make sure FConfigCacheIni::GetString (and friends) is thread safe from dynamic layer loading, and SafeReload
 // if set to 0, this will behave as 2025 and before (no safety for GetString)
-// note that if you bypass FConfigCache::Get* and use FCOnfigFile::Get*, then it will not be threadsafe
+// note that if you bypass FConfigCache::Get* and use FConfigFile::Get*, then it will not be threadsafe
 #define UE_THREAD_SAFETY_PROTECT_GETSTRING 1
 
 // if set to 1, GetString (GetInt, etc) will be slower, but SafeUnload will unload more memory
@@ -6649,6 +6649,7 @@ void FConfigCacheIni::InitializeConfigSystem()
 		// create GConfig
 		GConfig = new FConfigCacheIni(EConfigCacheType::DiskBacked, FPlatformProperties::IniPlatformName(), true /* bInGloballyRegistered */);
 
+	#if UE_ENABLE_CONFIG_SYSTEM
 		// create a context object that we will use for all of the main ini files
 		FConfigContext Context = FConfigContext::ReadIntoGConfig();
 
@@ -6680,6 +6681,7 @@ void FConfigCacheIni::InitializeConfigSystem()
 
 		// load editor, etc config files
 		LoadRemainingConfigFiles(Context);
+	#endif // UE_ENABLE_CONFIG_SYSTEM
 	}
 
 	FCoreDelegates::TSOnConfigSectionsChanged().AddStatic(OnConfigSectionsChanged);
diff --git a/Engine/Source/Runtime/Core/Private/Misc/DataDrivenPlatformInfoRegistry.cpp b/Engine/Source/Runtime/Core/Private/Misc/DataDrivenPlatformInfoRegistry.cpp
--- a/Engine/Source/Runtime/Core/Private/Misc/DataDrivenPlatformInfoRegistry.cpp
+++ b/Engine/Source/Runtime/Core/Private/Misc/DataDrivenPlatformInfoRegistry.cpp
@@ -258,6 +258,7 @@ static inline FName GetSectionString(const FConfigSection& Section, FName Key, F
 #if DDPI_HAS_EXTENDED_PLATFORMINFO_DATA
 static void ParsePreviewPlatforms(const FConfigFile& IniFile)
 {
+#if UE_ENABLE_CONFIG_SYSTEM
 	if (!FParse::Param(FCommandLine::Get(), TEXT("NoPreviewPlatforms")))
 	{
 		// walk over the file looking for PreviewPlatform sections
@@ -348,8 +349,9 @@ static void ParsePreviewPlatforms(const FConfigFile& IniFile)
 			}
 		}
 	}
+#endif // UE_ENABLE_CONFIG_SYSTEM
 }
-#endif
+#endif // DDPI_HAS_EXTENDED_PLATFORMINFO_DATA
 
 static void LoadDDPIIniSettings(const FConfigFile& IniFile, FDataDrivenPlatformInfo& Info, FName PlatformName)
 {
diff --git a/Engine/Source/Runtime/Core/Private/Misc/Paths.cpp b/Engine/Source/Runtime/Core/Private/Misc/Paths.cpp
--- a/Engine/Source/Runtime/Core/Private/Misc/Paths.cpp
+++ b/Engine/Source/Runtime/Core/Private/Misc/Paths.cpp
@@ -686,10 +686,9 @@ const TArray<FString>& FPaths::GetEngineLocalizationPaths()
 		if(GConfig && GConfig->IsReadyForUse())
 		{
 			GConfig->GetArray( TEXT("Internationalization"), TEXT("LocalizationPaths"), StaticData.EngineLocalizationPaths, GEngineIni );
-			if(!StaticData.EngineLocalizationPaths.Num())
-			{
-				UE_LOG(LogInit, Warning, TEXT("No paths for engine localization data were specifed in the engine configuration."));
-			}
+#if UE_ENABLE_ICU
+			UE_CLOG(!StaticData.EngineLocalizationPaths.Num(), LogInit, Warning, TEXT("No paths for engine localization data were specifed in the engine configuration."));
+#endif
 			StaticData.bEngineLocalizationPathsInitialized = true;
 		}
 		else
@@ -782,10 +781,9 @@ const TArray<FString>& FPaths::GetGameLocalizationPaths()
 		if(GConfig && GConfig->IsReadyForUse())
 		{
 			GConfig->GetArray( TEXT("Internationalization"), TEXT("LocalizationPaths"), StaticData.GameLocalizationPaths, GGameIni );
-			if(!StaticData.GameLocalizationPaths.Num()) // Failed to find localization path.
-			{
-				UE_LOG(LogPaths, Warning, TEXT("No paths for game localization data were specifed in the game configuration."));
-			}
+#if UE_ENABLE_ICU
+			UE_CLOG(!StaticData.GameLocalizationPaths.Num(), LogPaths, Warning, TEXT("No paths for game localization data were specifed in the game configuration."));
+#endif
 			StaticData.bGameLocalizationPathsInitialized = true;
 		}
 	}
