[AutoRTFM] Replace `AutoRTFM::IsClosed() || AutoRTFM::IsCommittingOrAborting()` with `!AutoRTFM::IsIdle()`

AutoRTFM::IsCommittingOrAborting() is a poorly named function that does not correctly describe what states we're querying.
The common pattern of `AutoRTFM::IsClosed() || AutoRTFM::IsCommittingOrAborting()` simplifies down to `!AutoRTFM::IsIdle()`, so replace these expressions with the simpiler form.

#Synced-CL 50252355

diff --git a/Engine/Source/Runtime/AutoRTFM/Public/AutoRTFM.h b/Engine/Source/Runtime/AutoRTFM/Public/AutoRTFM.h
--- a/Engine/Source/Runtime/AutoRTFM/Public/AutoRTFM.h
+++ b/Engine/Source/Runtime/AutoRTFM/Public/AutoRTFM.h
@@ -215,7 +215,8 @@ inline bool IsCascading(ETransactionStatus Status)
 // The context status shows what state the AutoRTFM context is currently in.
 enum class EContextStatus : uint8_t
 {
-	// An Idle status means we are not in transactional code.
+	// An Idle status means we are not in transactional code - i.e. the
+	// transaction stack is empty, or queried on a non-transactional thread.
 	Idle = autortfm_status_idle,
 
 	// An OnTrack status means we are in transactional code.
@@ -414,6 +415,13 @@ UE_AUTORTFM_CRITICAL_INLINE_ALWAYS bool IsTransactional()
 	return autortfm_is_context_status(autortfm_status_on_track);
 }
 
+// Returns true if we are currently outside of transactional code - i.e. the
+// transaction stack is empty, or called from a non-transactional thread.
+UE_AUTORTFM_CRITICAL_INLINE_ALWAYS bool IsIdle()
+{
+	return autortfm_is_context_status(autortfm_status_idle);
+}
+
 // Returns true if we are currently committing a transaction.
 // This will return true when inside an on-commit handler.
 UE_AUTORTFM_CRITICAL_INLINE_ALWAYS bool IsCommitting()
diff --git a/Engine/Source/Runtime/Core/Public/Async/TransactionallySafeRecursiveMutex.h b/Engine/Source/Runtime/Core/Public/Async/TransactionallySafeRecursiveMutex.h
--- a/Engine/Source/Runtime/Core/Public/Async/TransactionallySafeRecursiveMutex.h
+++ b/Engine/Source/Runtime/Core/Public/Async/TransactionallySafeRecursiveMutex.h
@@ -40,7 +40,7 @@ struct FTransactionallySafeRecursiveMutexDefinition final
 
 	void Lock()
 	{
-		if (AutoRTFM::IsTransactional() || AutoRTFM::IsCommittingOrAborting())
+		if (!AutoRTFM::IsIdle())
 		{
 			UE_AUTORTFM_OPEN
 			{
@@ -77,7 +77,7 @@ struct FTransactionallySafeRecursiveMutexDefinition final
 
 	void Unlock()
 	{
-		if (AutoRTFM::IsTransactional() || AutoRTFM::IsCommittingOrAborting())
+		if (!AutoRTFM::IsIdle())
 		{
 			// We explicitly copy the state here for the case that `this` was stack
 			// allocated and has already died before the on-commit is hit.
@@ -102,7 +102,7 @@ struct FTransactionallySafeRecursiveMutexDefinition final
 
 	[[nodiscard]] bool TryLock()
 	{
-		if (AutoRTFM::IsTransactional() || AutoRTFM::IsCommittingOrAborting())
+		if (!AutoRTFM::IsIdle())
 		{
 			// TODO: implement transactional TryLock. #jira SOL-8303
 			return false;
diff --git a/Engine/Source/Runtime/Core/Public/Async/TransactionallySafeSharedMutex.h b/Engine/Source/Runtime/Core/Public/Async/TransactionallySafeSharedMutex.h
--- a/Engine/Source/Runtime/Core/Public/Async/TransactionallySafeSharedMutex.h
+++ b/Engine/Source/Runtime/Core/Public/Async/TransactionallySafeSharedMutex.h
@@ -52,7 +52,7 @@ struct FTransactionallySafeSharedMutexDefinition
 
 	void LockShared()
 	{
-		if (AutoRTFM::IsTransactional() || AutoRTFM::IsCommittingOrAborting())
+		if (!AutoRTFM::IsIdle())
 		{
 			// Transactionally pessimise LockShared -> Lock.
 			Lock();
@@ -66,7 +66,7 @@ struct FTransactionallySafeSharedMutexDefinition
 
 	void UnlockShared()
 	{
-		if (AutoRTFM::IsTransactional() || AutoRTFM::IsCommittingOrAborting())
+		if (!AutoRTFM::IsIdle())
 		{
 			// Transactionally pessimise UnlockShared -> Unlock.
 			Unlock();
@@ -80,7 +80,7 @@ struct FTransactionallySafeSharedMutexDefinition
 
 	void Lock()
 	{
-		if (AutoRTFM::IsTransactional() || AutoRTFM::IsCommittingOrAborting())
+		if (!AutoRTFM::IsIdle())
 		{
 			UE_AUTORTFM_OPEN
 			{
@@ -110,7 +110,7 @@ struct FTransactionallySafeSharedMutexDefinition
 
 	void Unlock()
 	{
-		if (AutoRTFM::IsTransactional() || AutoRTFM::IsCommittingOrAborting())
+		if (!AutoRTFM::IsIdle())
 		{
 			// We explicitly copy the state here for the case that `this` was stack
 			// allocated and has already died before the on-commit is hit.
@@ -128,7 +128,7 @@ struct FTransactionallySafeSharedMutexDefinition
 
 	[[nodiscard]] bool TryLock()
 	{
-		if (AutoRTFM::IsTransactional() || AutoRTFM::IsCommittingOrAborting())
+		if (!AutoRTFM::IsIdle())
 		{
 			bool Ret = false;
 
diff --git a/Engine/Source/Runtime/Core/Public/Async/TransactionallySafeSharedRecursiveMutex.h b/Engine/Source/Runtime/Core/Public/Async/TransactionallySafeSharedRecursiveMutex.h
--- a/Engine/Source/Runtime/Core/Public/Async/TransactionallySafeSharedRecursiveMutex.h
+++ b/Engine/Source/Runtime/Core/Public/Async/TransactionallySafeSharedRecursiveMutex.h
@@ -52,7 +52,7 @@ struct FTransactionallySafeSharedRecursiveMutexDefinition
 
 	void LockShared(Core::Private::FSharedRecursiveMutexLink& Link)
 	{
-		if (AutoRTFM::IsTransactional() || AutoRTFM::IsCommittingOrAborting())
+		if (!AutoRTFM::IsIdle())
 		{
 			// Transactionally pessimise LockShared -> Lock.
 			Lock();
@@ -66,7 +66,7 @@ struct FTransactionallySafeSharedRecursiveMutexDefinition
 
 	void UnlockShared(Core::Private::FSharedRecursiveMutexLink& Link)
 	{
-		if (AutoRTFM::IsTransactional() || AutoRTFM::IsCommittingOrAborting())
+		if (!AutoRTFM::IsIdle())
 		{
 			// Transactionally pessimise UnlockShared -> Unlock.
 			Unlock();
@@ -80,7 +80,7 @@ struct FTransactionallySafeSharedRecursiveMutexDefinition
 
 	void Lock()
 	{
-		if (AutoRTFM::IsTransactional() || AutoRTFM::IsCommittingOrAborting())
+		if (!AutoRTFM::IsIdle())
 		{
 			UE_AUTORTFM_OPEN
 			{
@@ -117,7 +117,7 @@ struct FTransactionallySafeSharedRecursiveMutexDefinition
 
 	void Unlock()
 	{
-		if (AutoRTFM::IsTransactional() || AutoRTFM::IsCommittingOrAborting())
+		if (!AutoRTFM::IsIdle())
 		{
 			// We explicitly copy the state here for the case that `this` was stack
 			// allocated and has already died before the on-commit is hit.
diff --git a/Engine/Source/Runtime/Core/Public/Delegates/MulticastDelegateBase.h b/Engine/Source/Runtime/Core/Public/Delegates/MulticastDelegateBase.h
--- a/Engine/Source/Runtime/Core/Public/Delegates/MulticastDelegateBase.h
+++ b/Engine/Source/Runtime/Core/Public/Delegates/MulticastDelegateBase.h
@@ -362,7 +362,7 @@ private:
 		// skip compacting a multicast delegate when called via the AutoRTFM
 		// runtime. We skip calls both in the bodies of transactions themselves,
 		// and also those that are called as part of on-commit or on-abort calls.
-		return AutoRTFM::IsTransactional() || AutoRTFM::IsCommittingOrAborting();
+		return !AutoRTFM::IsIdle();
 	}
 
 	/**
diff --git a/Engine/Source/Runtime/Core/Public/Misc/MTTransactionallySafeAccessDetector.h b/Engine/Source/Runtime/Core/Public/Misc/MTTransactionallySafeAccessDetector.h
--- a/Engine/Source/Runtime/Core/Public/Misc/MTTransactionallySafeAccessDetector.h
+++ b/Engine/Source/Runtime/Core/Public/Misc/MTTransactionallySafeAccessDetector.h
@@ -48,7 +48,7 @@ private:
 	 */
 	inline bool AcquireReadAccess() const
 	{
-		if (AutoRTFM::IsTransactional() || AutoRTFM::IsCommittingOrAborting())
+		if (!AutoRTFM::IsIdle())
 		{
 			const bool bResult = AutoRTFM::Open([&]
 				{
@@ -92,7 +92,7 @@ private:
 	 */
 	inline bool ReleaseReadAccess() const
 	{
-		if (AutoRTFM::IsTransactional() || AutoRTFM::IsCommittingOrAborting())
+		if (!AutoRTFM::IsIdle())
 		{
 			RegisterOnCommitRelease();
 
@@ -113,7 +113,7 @@ private:
 	 */
 	inline bool AcquireWriteAccess() const
 	{
-		if (AutoRTFM::IsTransactional() || AutoRTFM::IsCommittingOrAborting())
+		if (!AutoRTFM::IsIdle())
 		{
 			const bool bResult = AutoRTFM::Open([&]
 				{
@@ -166,7 +166,7 @@ private:
 	 */
 	inline bool ReleaseWriteAccess() const
 	{
-		if (AutoRTFM::IsTransactional() || AutoRTFM::IsCommittingOrAborting())
+		if (!AutoRTFM::IsIdle())
 		{
 			RegisterOnCommitRelease();
 
diff --git a/Engine/Source/Runtime/Core/Public/Serialization/TransactionallySafeArchiveWriter.h b/Engine/Source/Runtime/Core/Public/Serialization/TransactionallySafeArchiveWriter.h
--- a/Engine/Source/Runtime/Core/Public/Serialization/TransactionallySafeArchiveWriter.h
+++ b/Engine/Source/Runtime/Core/Public/Serialization/TransactionallySafeArchiveWriter.h
@@ -112,7 +112,7 @@ private:
 		if (bRegisteredCommitHandler)
 		{
 			// If we already have a registered commit handler, we are either within a transaction or are committing/aborting.
-			check(AutoRTFM::IsTransactional() || AutoRTFM::IsCommittingOrAborting());
+			check(!AutoRTFM::IsIdle());
 			return true;
 		}
 		else if (AutoRTFM::IsTransactional())
diff --git a/Engine/Source/Runtime/Experimental/Chaos/Public/Framework/Threading.h b/Engine/Source/Runtime/Experimental/Chaos/Public/Framework/Threading.h
--- a/Engine/Source/Runtime/Experimental/Chaos/Public/Framework/Threading.h
+++ b/Engine/Source/Runtime/Experimental/Chaos/Public/Framework/Threading.h
@@ -1040,7 +1040,7 @@ FORCEINLINE void EnsureIsInGameThreadContext()
 
 		void ReadLock()
 		{
-			if (AutoRTFM::IsTransactional() || AutoRTFM::IsCommittingOrAborting())
+			if (!AutoRTFM::IsIdle())
 			{
 				// Transactionally pessimise ReadLock -> WriteLock.
 				WriteLock();
@@ -1054,7 +1054,7 @@ FORCEINLINE void EnsureIsInGameThreadContext()
 
 		void ReadUnlock()
 		{
-			if (AutoRTFM::IsTransactional() || AutoRTFM::IsCommittingOrAborting())
+			if (!AutoRTFM::IsIdle())
 			{
 				// Transactionally pessimise ReadUnlock -> WriteUnlock.
 				WriteUnlock();
@@ -1068,7 +1068,7 @@ FORCEINLINE void EnsureIsInGameThreadContext()
 
 		void WriteLock()
 		{
-			if (AutoRTFM::IsTransactional() || AutoRTFM::IsCommittingOrAborting())
+			if (!AutoRTFM::IsIdle())
 			{
 				UE_AUTORTFM_OPEN
 					{
@@ -1105,7 +1105,7 @@ FORCEINLINE void EnsureIsInGameThreadContext()
 
 		void WriteUnlock()
 		{
-			if (AutoRTFM::IsTransactional() || AutoRTFM::IsCommittingOrAborting())
+			if (!AutoRTFM::IsIdle())
 			{
 				// We explicitly copy the state here for the case that `this` was stack
 				// allocated and has already died before the on-commit is hit.
