[Core]
* Fixed so BlankProgram can compile on ios (can probably not run since it does not pull in ApplicationCore)

#Synced-CL 50312411

diff --git a/Engine/Source/Runtime/Core/Private/IOS/IOSPlatformMisc.cpp b/Engine/Source/Runtime/Core/Private/IOS/IOSPlatformMisc.cpp
--- a/Engine/Source/Runtime/Core/Private/IOS/IOSPlatformMisc.cpp
+++ b/Engine/Source/Runtime/Core/Private/IOS/IOSPlatformMisc.cpp
@@ -5,6 +5,7 @@
 =============================================================================*/
 
 #include "IOS/IOSPlatformMisc.h"
+
 #include "Apple/AppleStringUtils.h"
 #include "Apple/ApplePlatformMisc.h"
 #include "Async/TaskGraphInterfaces.h"
@@ -18,11 +19,8 @@
 #include "Internationalization/Regex.h"
 #include "IOSChunkInstaller.h"
 #include "IOS/IOSMallocZone.h"
-#include "IOS/IOSApplication.h"
-#include "IOS/IOSAppDelegate.h"
 #include "IOS/IOSPlatformCrashContext.h"
 #include "IOS/IOSPlatformPLCrashReporterIncludes.h"
-#include "IOS/IOSView.h"
 #include "Misc/App.h"
 #include "Misc/CommandLine.h"
 #include "Misc/ConfigCacheIni.h"
@@ -64,7 +62,13 @@
 #include "Apple/PostAppleSystemHeaders.h"
 #include "Apple/ScopeAutoreleasePool.h"
 
+#if WITH_APPLICATION_CORE
+#include "IOS/IOSApplication.h"
+#include "IOS/IOSAppDelegate.h"
+#include "IOS/IOSView.h"
 #include "IOSConfigRules.h"
+#endif
+
 #include <Metal/Metal.h>
 
 #if !defined ENABLE_ADVERTISING_IDENTIFIER
@@ -195,9 +199,13 @@ void FIOSPlatformMisc::PlatformInit()
 	// PlatformInit() starts the UI thread which creates the framebuffer and it requires
 	// "r.MobileContentScaleFactor" and "r.Mobile.DesiredResX/Y" to be available before 
 	// it's creation, so need to cache those value now.
+#if WITH_APPLICATION_CORE
 	[[IOSAppDelegate GetDelegate] LoadScreenResolutionModifiers];
 		
 	FAppEntry::PlatformInit();
+#else
+	check(false); // This code path is not implemented
+#endif
 
 	// Increase the maximum number of simultaneously open files
 	struct rlimit Limit;
@@ -407,33 +415,53 @@ const TCHAR* FIOSPlatformMisc::GamePersistentDownloadDir()
 
 EAppReturnType::Type FIOSPlatformMisc::MessageBoxExt( EAppMsgType::Type MsgType, const TCHAR* Text, const TCHAR* Caption )
 {
+#if WITH_APPLICATION_CORE
 	extern EAppReturnType::Type MessageBoxExtImpl( EAppMsgType::Type MsgType, const TCHAR* Text, const TCHAR* Caption );
 	return MessageBoxExtImpl(MsgType, Text, Caption);
+#else
+	return {};
+#endif
 }
 
 int FIOSPlatformMisc::GetAudioVolume()
 {
+#if WITH_APPLICATION_CORE
 	return [[IOSAppDelegate GetDelegate] GetAudioVolume];
+#else
+	return 0;
+#endif
 }
 
 int32 FIOSPlatformMisc::GetDeviceVolume()
 {
+#if WITH_APPLICATION_CORE
 	return [[IOSAppDelegate GetDelegate] GetAudioVolume];
+#else
+	return 0;
+#endif
 }
 
 bool FIOSPlatformMisc::AreHeadphonesPluggedIn()
 {
+#if WITH_APPLICATION_CORE
 	return [[IOSAppDelegate GetDelegate] AreHeadphonesPluggedIn];
+#else
+	return false;
+#endif
 }
 
 int FIOSPlatformMisc::GetBatteryLevel()
 {
+#if WITH_APPLICATION_CORE
 	return [[IOSAppDelegate GetDelegate] GetBatteryLevel];
+#else
+	return 0;
+#endif
 }
 
 float FIOSPlatformMisc::GetBrightness()
 {
-#if !PLATFORM_TVOS && !PLATFORM_VISIONOS
+#if !PLATFORM_TVOS && !PLATFORM_VISIONOS && WITH_APPLICATION_CORE
 	return (float)[[[IOSAppDelegate GetDelegate] window] screen].brightness;
 #else
 	return 1.0f;
@@ -442,7 +470,7 @@ float FIOSPlatformMisc::GetBrightness()
 
 void FIOSPlatformMisc::SetBrightness(float Brightness)
 {
-#if !PLATFORM_TVOS && !PLATFORM_VISIONOS
+#if !PLATFORM_TVOS && !PLATFORM_VISIONOS && WITH_APPLICATION_CORE
 	dispatch_async(dispatch_get_main_queue(), ^{
 		[[[IOSAppDelegate GetDelegate] window] screen].brightness = Brightness;
 	});
@@ -451,12 +479,16 @@ void FIOSPlatformMisc::SetBrightness(float Brightness)
 
 bool FIOSPlatformMisc::IsRunningOnBattery()
 {
+#if WITH_APPLICATION_CORE
 	return [[IOSAppDelegate GetDelegate] IsRunningOnBattery];
+#else
+	return false;
+#endif
 }
 
 float FIOSPlatformMisc::GetDeviceTemperatureLevel()
 {
-#if !PLATFORM_TVOS
+#if WITH_APPLICATION_CORE && !PLATFORM_TVOS
     switch ([[IOSAppDelegate GetDelegate] GetThermalState])
     {
         case NSProcessInfoThermalStateNominal:	return (float)FCoreDelegates::ETemperatureSeverity::Good; break;
@@ -561,7 +593,7 @@ UIInterfaceOrientationMask GetUIInterfaceOrientationMask(EDeviceScreenOrientatio
 
 EDeviceScreenOrientation FIOSPlatformMisc::GetDeviceOrientation()
 {
-#if !PLATFORM_TVOS && !PLATFORM_VISIONOS
+#if WITH_APPLICATION_CORE && !PLATFORM_TVOS && !PLATFORM_VISIONOS
 	IOSAppDelegate* AppDelegate = [IOSAppDelegate GetDelegate];
 	if (AppDelegate.InterfaceOrientation == UIInterfaceOrientationUnknown)
 	{
@@ -578,7 +610,7 @@ void FIOSPlatformMisc::SetAllowedDeviceOrientation(EDeviceScreenOrientation NewA
 {
 	AllowedDeviceOrientation = NewAllowedDeviceOrientation;
 	
-#if !PLATFORM_TVOS && !PLATFORM_VISIONOS
+#if WITH_APPLICATION_CORE && !PLATFORM_TVOS && !PLATFORM_VISIONOS
 	IOSAppDelegate* AppDelegate = [IOSAppDelegate GetDelegate];
 	AppDelegate.IOSView->SupportedInterfaceOrientations = GetUIInterfaceOrientationMask(NewAllowedDeviceOrientation);
 
@@ -1185,7 +1217,11 @@ void FIOSPlatformMisc::RequestStoreReview()
 
 bool FIOSPlatformMisc::IsUpdateAvailable()
 {
+#if WITH_APPLICATION_CORE 
 	return [[IOSAppDelegate GetDelegate] IsUpdateAvailable];
+#else
+	return false;
+#endif
 }
 
 /**
@@ -1195,7 +1231,7 @@ bool FIOSPlatformMisc::IsUpdateAvailable()
 */
 FString FIOSPlatformMisc::GetUniqueAdvertisingId()
 {
-#if !PLATFORM_TVOS && !PLATFORM_VISIONOS && ENABLE_ADVERTISING_IDENTIFIER
+#if WITH_APPLICATION_CORE && !PLATFORM_TVOS && !PLATFORM_VISIONOS && ENABLE_ADVERTISING_IDENTIFIER
 	// Check to see if this OS has this function
 	if ([[ASIdentifierManager sharedManager] respondsToSelector:@selector(advertisingIdentifier)])
 	{
@@ -1248,7 +1284,7 @@ class IPlatformChunkInstall* FIOSPlatformMisc::GetPlatformChunkInstall()
 
 bool FIOSPlatformMisc::SupportsForceTouchInput()
 {
-#if !PLATFORM_TVOS && !PLATFORM_VISIONOS
+#if WITH_APPLICATION_CORE && !PLATFORM_TVOS && !PLATFORM_VISIONOS
 	return [[[IOSAppDelegate GetDelegate].IOSView traitCollection] forceTouchCapability];
 #else
 	return false;
@@ -1363,6 +1399,7 @@ void FIOSPlatformMisc::ReleaseMobileHaptics()
 
 void FIOSPlatformMisc::ShareURL(const FString& URL, const FText& Description, int32 LocationHintX, int32 LocationHintY)
 {
+#if WITH_APPLICATION_CORE
 	NSString* SharedString = [NSString stringWithFString:Description.ToString()];
 	NSURL* SharedURL = [NSURL URLWithString:[NSString stringWithFString:URL]];
     
@@ -1403,6 +1440,7 @@ void FIOSPlatformMisc::ShareURL(const FString& URL, const FText& Description, in
 		}
 #endif // !PLATFORM_TVOS
 	});
+#endif // WITH_APPLICATION_CORE
 }
 
 
@@ -1478,21 +1516,32 @@ bool FIOSPlatformMisc::FileExistsInPlatformPackage(const FString& RelativePath)
 
 void FIOSPlatformMisc::EnableVoiceChat(bool bEnable)
 {
+#if WITH_APPLICATION_CORE
 	return [[IOSAppDelegate GetDelegate] EnableVoiceChat:bEnable];
+#endif
 }
 
 bool FIOSPlatformMisc::IsVoiceChatEnabled()
 {
+#if WITH_APPLICATION_CORE
 	return [[IOSAppDelegate GetDelegate] IsVoiceChatEnabled];
+#else
+	return false;
+#endif
 }
 
 bool FIOSPlatformMisc::HasRecordPermission()
 {
+#if WITH_APPLICATION_CORE
 	return [[IOSAppDelegate GetDelegate] HasRecordPermission];
+#else
+	return false;
+#endif
 }
 
 void FIOSPlatformMisc::RegisterForRemoteNotifications()
 {
+#if WITH_APPLICATION_CORE
 	if (FApp::IsUnattended())
 	{
 		return;
@@ -1526,6 +1575,7 @@ void FIOSPlatformMisc::RegisterForRemoteNotifications()
 							  }];
     });
 #endif
+#endif
 }
 
 bool FIOSPlatformMisc::IsRegisteredForRemoteNotifications()
@@ -1589,8 +1639,12 @@ FIOSPlatformMisc::EIOSAuthNotificationStatus FIOSPlatformMisc::GetNotificationAu
 
 FString FIOSPlatformMisc::GetPendingActivationProtocol()
 {
+#if WITH_APPLICATION_CORE
 	NSString* ActivationProtocol = [[IOSAppDelegate GetDelegate] GetPendingActivationProtocol];
 	return FString(ActivationProtocol);
+#else
+	return {};
+#endif
 }
 
 namespace UE::FIOSPlatformMisc::Entitlements
@@ -2177,7 +2231,12 @@ static TMap<FString, FString> GetPredefinedConfigRuleVars()
 
 const TMap<FString, FString>& FIOSPlatformMisc::GetConfigRuleVars()
 {
+#if WITH_APPLICATION_CORE
 	return FIOSConfigRules::GetConfigRulesMap();
+#else
+	static TMap<FString, FString> Map;
+	return Map;
+#endif
 }
 
 void FIOSPlatformMisc::PlatformPreInit()
@@ -2189,7 +2248,9 @@ void FIOSPlatformMisc::PlatformPreInit()
     // turn off SIGPIPE crashes
     signal(SIGPIPE, SIG_IGN);
 
+#if WITH_APPLICATION_CORE
 	FIOSConfigRules::Init(GetPredefinedConfigRuleVars());
+#endif
 }
 
 // Make sure that SetStoredValue and GetStoredValue generate the same key
@@ -2340,7 +2401,9 @@ void FIOSPlatformMisc::RequestExit(bool Force, const TCHAR* CallSite)
 	else
 	{
 		// ForceExit is sort of a misnomer here.  This will exit the engine loop before calling _Exit() from the app delegate
+#if WITH_APPLICATION_CORE
 		[[IOSAppDelegate GetDelegate] ForceExit];
+#endif
 	}
 }
 
@@ -2357,7 +2420,9 @@ void FIOSPlatformMisc::RequestExitWithStatus(bool Force, uint8 ReturnCode, const
 			Force, ReturnCode, CallSite ? CallSite : TEXT("<NoCallSiteInfo>"));
 
 		// ForceExit is sort of a misnomer here.  This will exit the engine loop before calling _Exit() from the app delegate
+#if WITH_APPLICATION_CORE
 		[[IOSAppDelegate GetDelegate] ForceExit];
+#endif
 	}
 }
 
diff --git a/Engine/Source/Runtime/Launch/Private/LaunchEngineLoop.cpp b/Engine/Source/Runtime/Launch/Private/LaunchEngineLoop.cpp
--- a/Engine/Source/Runtime/Launch/Private/LaunchEngineLoop.cpp
+++ b/Engine/Source/Runtime/Launch/Private/LaunchEngineLoop.cpp
@@ -307,7 +307,7 @@ class FFeedbackContext;
 #define WITH_CONFIG_PATCHING 0
 #endif
 
-#if PLATFORM_IOS || PLATFORM_TVOS
+#if WITH_APPLICATION_CORE && (PLATFORM_IOS || PLATFORM_TVOS)
 #include "IOS/IOSAppDelegate.h"
 #endif
 
@@ -6749,7 +6749,7 @@ bool FEngineLoop::AppInit( )
 	// Put the command line and config info into the suppression system (before plugins start loading)
 	FLogSuppressionInterface::Get().ProcessConfigAndCommandLine();
 
-#if PLATFORM_IOS || PLATFORM_TVOS
+#if WITH_APPLICATION_CORE && (PLATFORM_IOS || PLATFORM_TVOS)
 	// Now that the config system is ready, init the audio system.
 	[[IOSAppDelegate GetDelegate] InitializeAudioSession];
 #endif
