[Core] Compact set improvements
* Fix old naming from CompactHashTable to CompactHashTableView
* Rework remove function to do validations before modifying set for better debugging
* Fixup compact hash so it's properly validating lookup when relocating last element
* Normalize printing functions to use %d for num elements since we only support 32 bit for now
* Move out of memory check to cover all cases, and fixup some raw (u)int32 usages to use appropriate typedefs

#tests Used to help track down issues with improper set usage, ran many preflights

#Synced-CL 50170368

diff --git a/Engine/Source/Runtime/Core/Public/Containers/CompactHashTable.h b/Engine/Source/Runtime/Core/Public/Containers/CompactHashTable.h
--- a/Engine/Source/Runtime/Core/Public/Containers/CompactHashTable.h
+++ b/Engine/Source/Runtime/Core/Public/Containers/CompactHashTable.h
@@ -30,9 +30,9 @@
 
 namespace UE::Core::CompactHashTable::Private
 {
-	// Remove an element from the list and patch up any next index references, after this the spot is empty and needs to be filled
+	// The reference to the given index in chain of a given hash key and return a pointer to it
 	template<typename IndexType>
-	inline IndexType RemoveInternal(const uint32 Index, const uint32 Key, IndexType* HashData, const uint32 HashCount, IndexType* NextIndexData, const uint32 NextIndexCount)
+	inline IndexType* FindLookupEntry(const uint32 Index, const uint32 Key, IndexType* HashData, const uint32 HashCount, IndexType* NextIndexData, const uint32 NextIndexCount)
 	{
 		const uint32 HashIndex = Key & (HashCount - 1);
 
@@ -40,14 +40,12 @@ namespace UE::Core::CompactHashTable::Private
 		{
 			if (*Lookup == Index)
 			{
-				// Next = Next->Next
-				*Lookup = NextIndexData[Index];
-				return *Lookup;
+				return Lookup;
 			}
 		}
 
-		checkf(false, TEXT("Failed to find object to remove from set.")); // This usually indicates that something changed in the entry that affected the hash so it no longer matches up.
-		return (IndexType)INDEX_NONE;
+		checkf(false, TEXT("Failed to find entry from set.")); // This usually indicates that something changed in the entry that affected the hash so it no longer matches up.
+		return nullptr; // will crash
 	}
 }
 
@@ -164,28 +162,22 @@ namespace UE::Core::CompactHashTable
 	{
 		checkSlow(LastIndex < NextIndexCount && Index <= LastIndex);
 
-		IndexType NextIndex = UE::Core::CompactHashTable::Private::RemoveInternal<IndexType>(Index, Key, HashData, HashCount, NextIndexData, NextIndexCount);
+		// Next = Next->Next
+		IndexType* LookupEntry = UE::Core::CompactHashTable::Private::FindLookupEntry<IndexType>(Index, Key, HashData, HashCount, NextIndexData, NextIndexCount);
+		*LookupEntry = NextIndexData[Index];
 
 		if (Index != LastIndex)
 		{
 			// Remap the last element into the now empty spot
-			const uint32 HashIndex = OptLastKey & (HashCount - 1);
+			IndexType* LastLookupEntry = UE::Core::CompactHashTable::Private::FindLookupEntry<IndexType>(LastIndex, OptLastKey, HashData, HashCount, NextIndexData, NextIndexCount);
+			*LastLookupEntry = (IndexType)Index;
 
-			for (IndexType* Lookup = &HashData[HashIndex]; *Lookup < NextIndexCount; Lookup = &NextIndexData[*Lookup])
-			{
-				if (*Lookup == LastIndex)
-				{
-					// Fixup to look at now empty slot
-					*Lookup = (IndexType)Index;
-					// Now empty slot looks where the last item looked
-					NextIndexData[Index] = NextIndexData[LastIndex];
-					break;
-				}
-			}
+			// Move the last elements next index into the new slot as well
+			NextIndexData[Index] = NextIndexData[LastIndex];
 		}
 
-		// If the next index points to the last element, then it will get relocated into the current element so return the current instead of the next
-		return NextIndex == LastIndex ? Index : IndexTypeToSizeType(NextIndex);
+		// Return whatever next is at the original lookup entry that we fixed up. If it happened to be the last index the the previous block would have fixed it up for us
+		return IndexTypeToSizeType(*LookupEntry);
 	}
 
 	/* Remove an element from the list, shift all indexes down to preserve the order of elements in a list
@@ -207,8 +199,12 @@ namespace UE::Core::CompactHashTable
 			}
 		};
 
-		IndexType NextIndex = UE::Core::CompactHashTable::Private::RemoveInternal<IndexType>(Index, Key, HashData, HashCount, NextIndexData, NextIndexCount);
-		FixIndex(NextIndex);
+		IndexType* LookupEntry = UE::Core::CompactHashTable::Private::FindLookupEntry<IndexType>(Index, Key, HashData, HashCount, NextIndexData, NextIndexCount);
+		*LookupEntry = NextIndexData[Index];
+		FixIndex(*LookupEntry);
+
+		// Cache return value here since the memory location could move if it was a 'next index' that's located after where Index is
+		const uint32 NextIndex = IndexTypeToSizeType(*LookupEntry);
 
 		// For the hash indexes, just decrement any that are bigger than the removed element
 		for (uint32 HashIndex = 0; HashIndex < HashCount; ++HashIndex)
@@ -236,7 +232,7 @@ namespace UE::Core::CompactHashTable
 			}
 		}
 
-		return IndexTypeToSizeType(NextIndex);
+		return NextIndex;
 	}
 }
 
diff --git a/Engine/Source/Runtime/Core/Public/Containers/CompactSet.h.inl b/Engine/Source/Runtime/Core/Public/Containers/CompactSet.h.inl
--- a/Engine/Source/Runtime/Core/Public/Containers/CompactSet.h.inl
+++ b/Engine/Source/Runtime/Core/Public/Containers/CompactSet.h.inl
@@ -614,7 +614,7 @@ public:
 		// Template property, branch will be optimized out
 		if constexpr (AllocatorType::RequireRangeCheck)
 		{
-			checkf(IsValidId(Id), TEXT("Array index out of bounds: %d into an array of size %lld"), Id.AsInteger(), (long long)this->NumElements); // & for one branch
+			checkf(IsValidId(Id), TEXT("Array index out of bounds: %d into an array of size %d"), Id.AsInteger(), this->NumElements);
 		}
 	}
 
@@ -1147,13 +1147,13 @@ public:
 	{
 		if (this->MaxElements == 0)
 		{
-			Ar.Logf(TEXT("UE_TCOMPACT_SET: empty"), this->NumElements, this->MaxElements);
+			Ar.Log(TEXT("UE_TCOMPACT_SET: empty"));
 		}
 		else if (this->MaxElements > 0)
 		{
 			FConstCompactHashTableView HashTable = this->GetConstHashTableView();
 
-			Ar.Logf(TEXT("UE_TCOMPACT_SET: %i elements, %i max elements, %i hash slots"), this->NumElements, this->MaxElements, HashTable.GetHashCount());
+			Ar.Logf(TEXT("UE_TCOMPACT_SET: %d elements, %d max elements, %d hash slots"), this->NumElements, this->MaxElements, HashTable.GetHashCount());
 
 			for (uint32 HashIndex = 0; HashIndex < HashTable.GetHashCount(); ++HashIndex)
 			{
@@ -1495,42 +1495,45 @@ private:
 	template <bool IsStable = false>
 	void RemoveByIndex(const SizeType ElementIndex)
 	{
-		checkf(ElementIndex >= 0 && ElementIndex < this->NumElements, TEXT("Invalid ElementIndex passed to UE_TCOMPACT_SET::RemoveByIndex"));
-		RemoveByIndexAndHash<IsStable>(ElementIndex, KeyFuncs::GetKeyHash(KeyFuncs::GetSetKey(GetData()[ElementIndex])));
-	}
-
-	template <bool IsStable = false>
-	SizeType RemoveByIndexAndHash(const SizeType ElementIndex, const uint32 KeyHash)
-	{
-		checkf(ElementIndex >= 0 && ElementIndex < this->NumElements, TEXT("Invalid ElementIndex passed to UE_TCOMPACT_SET::RemoveByIndex"));
+		checkf(ElementIndex >= 0 && ElementIndex < this->NumElements, TEXT("Invalid ElementIndex (%d) passed to UE_TCOMPACT_SET::RemoveByIndex. NumElements=%d"), ElementIndex, this->NumElements);
 
 		ElementType* ElementsData = GetData();
 		FCompactHashTableView HashTable = GetHashTableView();
 
-		--this->NumElements;
+		RemoveByIndexAndHash<IsStable>(ElementIndex, KeyFuncs::GetKeyHash(KeyFuncs::GetSetKey(ElementsData[ElementIndex])), ElementsData, HashTable);
+	}
 
-		const SizeType LastElementIndex = this->NumElements;
+	template <bool IsStable = false>
+	SizeType RemoveByIndexAndHash(const SizeType ElementIndex, const uint32 KeyHash, ElementType* ElementsData, FCompactHashTableView& HashTable)
+	{
+		const SizeType LastElementIndex = this->NumElements - 1;
+		SizeType NextIndex; 
+		
+		// Note: Destructing/moving elements for better debugability
 		if (ElementIndex == LastElementIndex)
 		{
+			NextIndex = HashTable.Remove(ElementIndex, KeyHash, ElementIndex, KeyHash);
+			
 			ElementsData[LastElementIndex].~ElementType();
-			return HashTable.Remove(ElementIndex, KeyHash, ElementIndex, KeyHash);
 		}
-
-		if constexpr (IsStable)
+		else if constexpr (IsStable)
 		{
+			NextIndex = HashTable.RemoveStable(ElementIndex, KeyHash);
+
 			ElementsData[ElementIndex].~ElementType();
 			RelocateConstructItems<ElementType>(ElementsData + ElementIndex, ElementsData + ElementIndex + 1, LastElementIndex - ElementIndex);
-
-			return HashTable.RemoveStable(ElementIndex, KeyHash);
 		}
 		else
 		{
 			const uint32 LastElementHash = KeyFuncs::GetKeyHash(KeyFuncs::GetSetKey(ElementsData[LastElementIndex]));
-
+			NextIndex = HashTable.Remove(ElementIndex, KeyHash, LastElementIndex, LastElementHash);
+			
 			MoveByRelocate(ElementsData[ElementIndex], ElementsData[LastElementIndex]);
-
-			return HashTable.Remove(ElementIndex, KeyHash, LastElementIndex, LastElementHash);
 		}
+
+		--this->NumElements;
+
+		return NextIndex;
 	}
 
 	template<typename ComparableKey, bool IsStable = false>
@@ -1539,24 +1542,25 @@ private:
 		checkSlow(this->NumElements > 0);
 		int32 NumRemovedElements = 0;
 
-		const ElementType* ElementsData = GetData();
+		ElementType* ElementsData = GetData();
 		FCompactHashTableView HashTable = GetHashTableView();
 
 		SizeType ElementIndex = HashTable.GetFirst(KeyHash);
 
 		while (ElementIndex != INDEX_NONE)
 		{
+			checkf(ElementIndex >= 0 && ElementIndex < this->NumElements, TEXT("Invalid ElementIndex (%d) passed to UE_TCOMPACT_SET::RemoveByIndexAndHash. NumElements=%d"), ElementIndex, this->NumElements);
 			if (KeyFuncs::Matches(KeyFuncs::GetSetKey(ElementsData[ElementIndex]), Key))
 			{
 				if constexpr (!KeyFuncs::bAllowDuplicateKeys)
 				{
-					RemoveByIndexAndHash<IsStable>(ElementIndex, KeyHash);
+					RemoveByIndexAndHash<IsStable>(ElementIndex, KeyHash, ElementsData, HashTable);
 					// If the hash disallows duplicate keys, we're done removing after the first matched key.
 					return 1;
 				}
 				else
 				{
-					ElementIndex = RemoveByIndexAndHash<IsStable>(ElementIndex, KeyHash);
+					ElementIndex = RemoveByIndexAndHash<IsStable>(ElementIndex, KeyHash, ElementsData, HashTable);
 					NumRemovedElements++;
 				}
 			}
diff --git a/Engine/Source/Runtime/Core/Public/Containers/CompactSetBase.h b/Engine/Source/Runtime/Core/Public/Containers/CompactSetBase.h
--- a/Engine/Source/Runtime/Core/Public/Containers/CompactSetBase.h
+++ b/Engine/Source/Runtime/Core/Public/Containers/CompactSetBase.h
@@ -38,19 +38,19 @@ public:
 	}
 
 	/** @return the number of elements. */
-	[[nodiscard]] UE_FORCEINLINE_HINT int32 Num() const
+	[[nodiscard]] UE_FORCEINLINE_HINT SizeType Num() const
 	{
 		return NumElements;
 	}
 
 	/** @return The number of elements the set can hold before reallocation. */
-	[[nodiscard]] UE_FORCEINLINE_HINT int32 Max() const
+	[[nodiscard]] UE_FORCEINLINE_HINT SizeType Max() const
 	{
 		return MaxElements;
 	}
 
 	/** Deprecated - unnecessary, keeping this here so TCompactSet can be swapped with TSet without changing code */
-	[[nodiscard]] UE_FORCEINLINE_HINT int32 GetMaxIndex() const
+	[[nodiscard]] UE_FORCEINLINE_HINT SizeType GetMaxIndex() const
 	{
 		return NumElements;
 	}
@@ -149,9 +149,9 @@ protected:
 		return FMath::Min<SizeType>(MaxElements, RealAvailableBytes / Layout.Size); // Now we can get the true number of elements we could potentially use within the aligned space
 	}
 
-	[[nodiscard]] int32 AllocatorCalculateSlackGrow(int32 NewMaxElements, const FCompactSetLayout& Layout) const
+	[[nodiscard]] SizeType AllocatorCalculateSlackGrow(SizeType NewMaxElements, const FCompactSetLayout& Layout) const
 	{
-		const size_t OldHashCount = MaxElements > 0 ? *(uint32*)(Elements.GetAllocation() + GetElementsSizeInBytes(MaxElements, Layout)) : 0;
+		const size_t OldHashCount = MaxElements > 0 ? *(HashCountType*)(Elements.GetAllocation() + GetElementsSizeInBytes(MaxElements, Layout)) : 0;
 		const size_t OldSize = MaxElements > 0 ? GetTotalMemoryRequiredInBytes(MaxElements, OldHashCount, Layout) : 0;
 
 		const size_t NewHashCount = NewMaxElements > 0 ? GetHashCount(NewMaxElements) : 0;
@@ -196,28 +196,28 @@ protected:
 		return SlackNumElements;
 	}
 
-	void ResizeAllocation(const int32 NewMaxElements, const FCompactSetLayout& Layout)
+	void ResizeAllocation(const SizeType NewMaxElements, const FCompactSetLayout& Layout)
 	{
 		(void)ResizeAllocationPreserveData(NewMaxElements, Layout, false);
 	}
 
 	// Use this is you'll be keeping the element data
-	[[nodiscard]] bool ResizeAllocationPreserveData(const int32 NewMaxElements, const FCompactSetLayout& Layout, bool bPreserve = true)
+	[[nodiscard]] bool ResizeAllocationPreserveData(const SizeType NewMaxElements, const FCompactSetLayout& Layout, bool bPreserve = true)
 	{
 		bool bRequiresRehash = false;
 
 		if (NewMaxElements != MaxElements)
 		{
-			const size_t OldHashCount = MaxElements > 0 ? *(uint32*)(Elements.GetAllocation() + GetElementsSizeInBytes(MaxElements, Layout)) : 0;
+			const size_t OldHashCount = MaxElements > 0 ? *(HashCountType*)(Elements.GetAllocation() + GetElementsSizeInBytes(MaxElements, Layout)) : 0;
 			const size_t OldSize = MaxElements > 0 ? GetTotalMemoryRequiredInBytes(MaxElements, OldHashCount, Layout) : 0;
 
 			const size_t NewHashCount = NewMaxElements > 0 ? UE::Core::CompactHashTable::GetHashCount(NewMaxElements) : 0;
 			const size_t NewSize = NewMaxElements > 0 ? GetTotalMemoryRequiredInBytes(NewMaxElements, NewHashCount, Layout) : 0;
 
+			checkf(NewSize >= 0 && NewSize <= MAX_int32, TEXT("Invalid size for TSet[%lld]: NewMaxElements[%d] ElementSize[%d] HashCount[%lld]"), NewSize, NewMaxElements, Layout.Size, NewHashCount);
+
 			if (bPreserve && NewMaxElements > MaxElements && OldHashCount == NewHashCount)
 			{
-				checkf(NewSize >= 0 && NewSize <= MAX_int32, TEXT("Invalid size for TSet[%d]: NewMaxElements[%d] ElementSize[%d] HashCount[%d]"), NewSize, NewMaxElements, Layout.Size, NewHashCount);
-
 				// If preserving then we copy over all the data
 				if constexpr (TAllocatorTraits<AllocatorType>::SupportsElementAlignment)
 				{
@@ -260,7 +260,7 @@ protected:
 				// Update hash size
 				if (NewMaxElements > 0)
 				{
-					uint32* NewHashTable = (uint32*)(Elements.GetAllocation() + GetElementsSizeInBytes(NewMaxElements, Layout));
+					HashCountType* NewHashTable = (HashCountType*)(Elements.GetAllocation() + GetElementsSizeInBytes(NewMaxElements, Layout));
 
 					*NewHashTable = NewHashCount;
 					bRequiresRehash = true;
diff --git a/Engine/Source/Runtime/Core/Tests/Containers/CompactHashTableTests.cpp b/Engine/Source/Runtime/Core/Tests/Containers/CompactHashTableTests.cpp
--- a/Engine/Source/Runtime/Core/Tests/Containers/CompactHashTableTests.cpp
+++ b/Engine/Source/Runtime/Core/Tests/Containers/CompactHashTableTests.cpp
@@ -188,7 +188,7 @@ TEST_CASE_NAMED(FStaticCompactHashTableTest, "System::Core::Containers::CompactH
 	}
 }
 
-TEST_CASE_NAMED(FDynamicCompactHashTableTest, "System::Core::Containers::CompactHashTable::Dynamic", "[Core][Containers][CompactHashTable]")
+TEST_CASE_NAMED(FCompactHashTableViewTest, "System::Core::Containers::CompactHashTable::Dynamic", "[Core][Containers][CompactHashTable]")
 {
 	SECTION("Constructor Behavior")
 	{
