Updated NoExportTypes.h:
- Synced differences in EPixelFormat.
- Added underlying types to forward declarations.
- Forward declarations for StaticEnum (important for future changes).
- Moved StaticEnum to Core so that specializations can be made for NoExportTypes.

#Synced-CL 50307883

diff --git a/Engine/Plugins/Animation/GameplayInsights/Source/RewindDebugger/Private/SModalSessionBrowser.h b/Engine/Plugins/Animation/GameplayInsights/Source/RewindDebugger/Private/SModalSessionBrowser.h
--- a/Engine/Plugins/Animation/GameplayInsights/Source/RewindDebugger/Private/SModalSessionBrowser.h
+++ b/Engine/Plugins/Animation/GameplayInsights/Source/RewindDebugger/Private/SModalSessionBrowser.h
@@ -12,7 +12,7 @@ class FLiveSessionTracker;
 class SModalSessionNameListPicker;
 class SEditableTextBox;
 
-namespace EAppReturnType { enum Type; }
+namespace EAppReturnType { enum Type : int; }
 
 /** Structure containing info about a Trace Session */
 
diff --git a/Engine/Plugins/ChaosVD/Source/ChaosVD/Private/Widgets/SChaosVDBrowseSessionsModal.h b/Engine/Plugins/ChaosVD/Source/ChaosVD/Private/Widgets/SChaosVDBrowseSessionsModal.h
--- a/Engine/Plugins/ChaosVD/Source/ChaosVD/Private/Widgets/SChaosVDBrowseSessionsModal.h
+++ b/Engine/Plugins/ChaosVD/Source/ChaosVD/Private/Widgets/SChaosVDBrowseSessionsModal.h
@@ -13,7 +13,7 @@ class FLiveSessionTracker;
 class SChaosVDNameListPicker;
 class SEditableTextBox;
 
-namespace EAppReturnType { enum Type; }
+namespace EAppReturnType { enum Type : int; }
 
 /** Structure containing info about a Trace Session */
 struct FChaosVDTraceSessionInfo
diff --git a/Engine/Plugins/Runtime/Harmonix/Source/HarmonixMidiEditor/Private/AssetDefinition_MidiFile.h b/Engine/Plugins/Runtime/Harmonix/Source/HarmonixMidiEditor/Private/AssetDefinition_MidiFile.h
--- a/Engine/Plugins/Runtime/Harmonix/Source/HarmonixMidiEditor/Private/AssetDefinition_MidiFile.h
+++ b/Engine/Plugins/Runtime/Harmonix/Source/HarmonixMidiEditor/Private/AssetDefinition_MidiFile.h
@@ -8,7 +8,7 @@
 
 struct FToolMenuContext;
 class UContentBrowserAssetContextMenuContext;
-namespace EAppReturnType { enum Type; }
+namespace EAppReturnType { enum Type : int; }
 
 UCLASS()
 class UAssetDefinition_MidiFile : public UAssetDefinitionDefault
diff --git a/Engine/Plugins/Tests/RHITests/Source/RHITests/Public/RHITextureTests.h b/Engine/Plugins/Tests/RHITests/Source/RHITests/Public/RHITextureTests.h
--- a/Engine/Plugins/Tests/RHITests/Source/RHITests/Public/RHITextureTests.h
+++ b/Engine/Plugins/Tests/RHITests/Source/RHITests/Public/RHITextureTests.h
@@ -849,7 +849,7 @@ public:
 			}
 		};
 #define RUN_TEST_ON_FORMAT(Format) RunTestOnFormat(Format);
-		FOREACH_ENUM_EPIXELFORMAT(RUN_TEST_ON_FORMAT);
+		FOREACH_ENUM_EPIXELFORMAT_CORE(RUN_TEST_ON_FORMAT);
 #undef RUN_TEST_ON_FORMAT
 
 		return bResult;
diff --git a/Engine/Source/Developer/TextureBuildUtilities/Private/TextureBuildUtilities.cpp b/Engine/Source/Developer/TextureBuildUtilities/Private/TextureBuildUtilities.cpp
--- a/Engine/Source/Developer/TextureBuildUtilities/Private/TextureBuildUtilities.cpp
+++ b/Engine/Source/Developer/TextureBuildUtilities/Private/TextureBuildUtilities.cpp
@@ -144,7 +144,7 @@ TEXTUREBUILDUTILITIES_API ERawImageFormat::Type GetVirtualTextureBuildIntermedia
 static EPixelFormat GetPixelFormatFromUtf8(const FUtf8StringView& InPixelFormatStr)
 {
 #define TEXT_TO_PIXELFORMAT(f) TEXT_TO_ENUM(f, InPixelFormatStr);
-	FOREACH_ENUM_EPIXELFORMAT(TEXT_TO_PIXELFORMAT)
+	FOREACH_ENUM_EPIXELFORMAT_CORE(TEXT_TO_PIXELFORMAT)
 #undef TEXT_TO_PIXELFORMAT
 		return PF_Unknown;
 }
diff --git a/Engine/Source/Runtime/ApplicationCore/Public/GenericPlatform/ICursor.h b/Engine/Source/Runtime/ApplicationCore/Public/GenericPlatform/ICursor.h
--- a/Engine/Source/Runtime/ApplicationCore/Public/GenericPlatform/ICursor.h
+++ b/Engine/Source/Runtime/ApplicationCore/Public/GenericPlatform/ICursor.h
@@ -4,6 +4,7 @@
 
 #include "CoreTypes.h"
 #include "Math/Vector2D.h"
+#include "UObject/CoreReflectedTypeAccessors.h"
 
 struct tagRECT;
 typedef struct tagRECT RECT;
@@ -66,6 +67,10 @@ namespace EMouseCursor
 	};
 }
 
+#ifdef COREUOBJECT_API
+template<> COREUOBJECT_API UEnum* StaticEnum<EMouseCursor::Type>();
+#endif
+
 class ICursor
 {
 public:
diff --git a/Engine/Source/Runtime/ApplicationCore/Public/GenericPlatform/IInputInterface.h b/Engine/Source/Runtime/ApplicationCore/Public/GenericPlatform/IInputInterface.h
--- a/Engine/Source/Runtime/ApplicationCore/Public/GenericPlatform/IInputInterface.h
+++ b/Engine/Source/Runtime/ApplicationCore/Public/GenericPlatform/IInputInterface.h
@@ -6,6 +6,7 @@
 #include "Containers/Array.h"
 #include "Math/Color.h"
 #include "UObject/NameTypes.h"
+#include "UObject/CoreReflectedTypeAccessors.h"
 
 
 // General identifiers for potential force feedback channels. These will be mapped according to the
@@ -103,6 +104,10 @@ enum class EInputDeviceTriggerMask : uint8
 };
 ENUM_CLASS_FLAGS(EInputDeviceTriggerMask)
 
+#ifdef COREUOBJECT_API
+template<> COREUOBJECT_API UEnum* StaticEnum<EInputDeviceTriggerMask>();
+#endif
+
 /**
  * Represents input device analog sticks that are available
  *
@@ -116,6 +121,10 @@ enum class EInputDeviceAnalogStickMask : uint8
 };
 ENUM_CLASS_FLAGS(EInputDeviceAnalogStickMask)
 
+#ifdef COREUOBJECT_API
+template<> COREUOBJECT_API UEnum* StaticEnum<EInputDeviceAnalogStickMask>();
+#endif
+
 struct FInputDeviceProperty
 {
 	FInputDeviceProperty(FName InName)
diff --git a/Engine/Source/Runtime/ApplicationCore/Public/GenericPlatform/InputDeviceMappingPolicy.h b/Engine/Source/Runtime/ApplicationCore/Public/GenericPlatform/InputDeviceMappingPolicy.h
--- a/Engine/Source/Runtime/ApplicationCore/Public/GenericPlatform/InputDeviceMappingPolicy.h
+++ b/Engine/Source/Runtime/ApplicationCore/Public/GenericPlatform/InputDeviceMappingPolicy.h
@@ -3,6 +3,7 @@
 #pragma once
 
 #include "HAL/Platform.h"	// for int32
+#include "UObject/CoreReflectedTypeAccessors.h"
 
 /**
  * The input device mapping policy controls how Human Interface Devices (HID's)
@@ -89,3 +90,6 @@ enum class EInputDeviceMappingPolicy : int32
 	MapAllDevicesToPrimaryUser = 3
 };
 
+#ifdef COREUOBJECT_API
+template<> COREUOBJECT_API UEnum* StaticEnum<EInputDeviceMappingPolicy>();
+#endif
diff --git a/Engine/Source/Runtime/Core/Private/Misc/PixelFormat.cpp b/Engine/Source/Runtime/Core/Private/Misc/PixelFormat.cpp
--- a/Engine/Source/Runtime/Core/Private/Misc/PixelFormat.cpp
+++ b/Engine/Source/Runtime/Core/Private/Misc/PixelFormat.cpp
@@ -252,7 +252,7 @@ const TCHAR* GetPixelFormatString(EPixelFormat InPixelFormat)
 {
 	switch (InPixelFormat)
 	{
-		FOREACH_ENUM_EPIXELFORMAT(CASE_ENUM_TO_TEXT)
+		FOREACH_ENUM_EPIXELFORMAT_CORE(CASE_ENUM_TO_TEXT)
 	default:
 		return TEXT("PF_Unknown");
 	}
@@ -261,7 +261,7 @@ const TCHAR* GetPixelFormatString(EPixelFormat InPixelFormat)
 EPixelFormat GetPixelFormatFromString(const TCHAR* InPixelFormatStr)
 {
 #define TEXT_TO_PIXELFORMAT(f) TEXT_TO_ENUM(f, InPixelFormatStr);
-	FOREACH_ENUM_EPIXELFORMAT(TEXT_TO_PIXELFORMAT)
+	FOREACH_ENUM_EPIXELFORMAT_CORE(TEXT_TO_PIXELFORMAT)
 #undef TEXT_TO_PIXELFORMAT
 		return PF_Unknown;
 }
diff --git a/Engine/Source/Runtime/Core/Public/GenericPlatform/GenericPlatformMisc.h b/Engine/Source/Runtime/Core/Public/GenericPlatform/GenericPlatformMisc.h
--- a/Engine/Source/Runtime/Core/Public/GenericPlatform/GenericPlatformMisc.h
+++ b/Engine/Source/Runtime/Core/Public/GenericPlatform/GenericPlatformMisc.h
@@ -12,6 +12,7 @@
 #include "Misc/EnumClassFlags.h"
 #include "ProfilingDebugging/CsvProfilerConfig.h"
 #include "ProfilingDebugging/ExternalProfilerConfig.h"
+#include "UObject/CoreReflectedTypeAccessors.h"
 
 #if PLATFORM_CPU_X86_FAMILY
 #include <xmmintrin.h> // _mm_prefetch
@@ -296,13 +297,17 @@ enum class EAppMsgCategory : uint8
 	Info,
 };
 
+#ifdef COREUOBJECT_API
+template<> COREUOBJECT_API UEnum* StaticEnum<EAppMsgCategory>();
+#endif
+
 
 namespace EAppMsgType
 {
 	/**
 	 * Enumerates supported message dialog button types.
 	 */
-	enum Type
+	enum Type : int
 	{
 		Ok,
 		YesNo,
@@ -315,13 +320,17 @@ namespace EAppMsgType
 	};
 }
 
+#ifdef COREUOBJECT_API
+template<> COREUOBJECT_API UEnum* StaticEnum<EAppMsgType::Type>();
+#endif
+
 
 namespace EAppReturnType
 {
 	/**
 	 * Enumerates message dialog return types.
 	 */
-	enum Type
+	enum Type : int
 	{
 		No,
 		Yes,
@@ -334,6 +343,10 @@ namespace EAppReturnType
 	};
 }
 
+#ifdef COREUOBJECT_API
+template<> COREUOBJECT_API UEnum* StaticEnum<EAppReturnType::Type>();
+#endif
+
 /**
  * Returns the string representation of the specified EAppReturnType::Type value.
  *
@@ -1382,7 +1395,7 @@ public:
 	 *
 	 * NOTE: This will not take affect immediately. An engine reboot is required to determine which input APIs should be created.
 	 * 
-	 * @param DeviceName The full comma-seperated list of preferred input APIs this application should use.
+	 * @param DeviceName The full comma-separated list of preferred input APIs this application should use.
 	 */
 	static CORE_API void SetPreferredInputDevices(const TCHAR* DeviceName);
 
@@ -1394,7 +1407,7 @@ public:
 	static CORE_API void ResetPreferredInputDevicesToDefault();
 
 	/**
-	 * @return A comma-seperated string value which represents a list of the preferred input APIs.
+	 * @return A comma-separated string value which represents a list of the preferred input APIs.
 	 */
 	static CORE_API FString GetCurrentPreferredInputDevices();
 
diff --git a/Engine/Source/Runtime/Core/Public/Internationalization/LocalizedTextSourceTypes.h b/Engine/Source/Runtime/Core/Public/Internationalization/LocalizedTextSourceTypes.h
--- a/Engine/Source/Runtime/Core/Public/Internationalization/LocalizedTextSourceTypes.h
+++ b/Engine/Source/Runtime/Core/Public/Internationalization/LocalizedTextSourceTypes.h
@@ -8,6 +8,7 @@
 #include "Templates/RefCounting.h"
 #include "Templates/Requires.h"
 #include "Templates/SharedPointer.h"
+#include "UObject/CoreReflectedTypeAccessors.h"
 
 #include <type_traits>
 
@@ -22,6 +23,10 @@ enum class ELocalizedTextSourceCategory : uint8
 	Editor,
 };
 
+#ifdef COREUOBJECT_API
+template<> COREUOBJECT_API UEnum* StaticEnum<ELocalizedTextSourceCategory>();
+#endif
+
 /**
  * Result codes from calling QueryLocalizedResourceResult.
  */
diff --git a/Engine/Source/Runtime/Core/Public/Math/Axis.h b/Engine/Source/Runtime/Core/Public/Math/Axis.h
--- a/Engine/Source/Runtime/Core/Public/Math/Axis.h
+++ b/Engine/Source/Runtime/Core/Public/Math/Axis.h
@@ -3,11 +3,12 @@
 #pragma once
 
 #include "HAL/Platform.h"
+#include "UObject/CoreReflectedTypeAccessors.h"
 
 // Generic axis enum (mirrored for property use in Object.h)
 namespace EAxis
 {
-	enum Type
+	enum Type : int
 	{
 		None,
 		X,
@@ -16,11 +17,15 @@ namespace EAxis
 	};
 }
 
+#ifdef COREUOBJECT_API
+template<> COREUOBJECT_API UEnum* StaticEnum<EAxis::Type>();
+#endif
+
 
 // Extended axis enum for more specialized usage
 namespace EAxisList
 {
-	enum Type
+	enum Type : int
 	{
 		None		= 0,
 		X			= 1 << 0,
@@ -51,6 +56,10 @@ namespace EAxisList
 	};
 }
 
+#ifdef COREUOBJECT_API
+template<> COREUOBJECT_NON_ATTRIBUTED_API UEnum* StaticEnum<EAxisList::Type>();
+#endif
+
 namespace EAxis
 {
 	// Converts the given AxisList value to the corresponding Axis value. Assumes a single axis within the provided AxisList.
diff --git a/Engine/Source/Runtime/Core/Public/Math/InterpCurvePoint.h b/Engine/Source/Runtime/Core/Public/Math/InterpCurvePoint.h
--- a/Engine/Source/Runtime/Core/Public/Math/InterpCurvePoint.h
+++ b/Engine/Source/Runtime/Core/Public/Math/InterpCurvePoint.h
@@ -11,8 +11,9 @@
 #include "Math/Vector.h"
 #include "Math/Quat.h"
 #include "Math/TwoVectors.h"
+#include "UObject/CoreReflectedTypeAccessors.h"
 
-enum EInterpCurveMode
+enum EInterpCurveMode : int
 {
 	/** A straight line between two keypoint values. */
 	CIM_Linear,
@@ -38,6 +39,10 @@ enum EInterpCurveMode
 	CIM_Unknown
 };
 
+#ifdef COREUOBJECT_API
+template<> COREUOBJECT_API UEnum* StaticEnum<EInterpCurveMode>();
+#endif
+
 
 /**
  * Template for interpolation points.
diff --git a/Engine/Source/Runtime/Core/Public/Math/RangeBound.h b/Engine/Source/Runtime/Core/Public/Math/RangeBound.h
--- a/Engine/Source/Runtime/Core/Public/Math/RangeBound.h
+++ b/Engine/Source/Runtime/Core/Public/Math/RangeBound.h
@@ -8,13 +8,14 @@
 #include "Containers/EnumAsByte.h"
 #include "Misc/FrameNumber.h"
 #include "Misc/DateTime.h"
+#include "UObject/CoreReflectedTypeAccessors.h"
 
 namespace ERangeBoundTypes
 {
 	/**
 	 * Enumerates the valid types of range bounds.
 	 */
-	enum Type
+	enum Type : int
 	{
 		/** The range excludes the bound. */
 		Exclusive,
@@ -27,6 +28,10 @@ namespace ERangeBoundTypes
 	};
 }
 
+#ifdef COREUOBJECT_API
+template<> COREUOBJECT_API UEnum* StaticEnum<ERangeBoundTypes::Type>();
+#endif
+
 
 /**
  * Template for range bounds.
diff --git a/Engine/Source/Runtime/Core/Public/Math/UnitConversion.h b/Engine/Source/Runtime/Core/Public/Math/UnitConversion.h
--- a/Engine/Source/Runtime/Core/Public/Math/UnitConversion.h
+++ b/Engine/Source/Runtime/Core/Public/Math/UnitConversion.h
@@ -9,6 +9,7 @@
 #include "Internationalization/Text.h"
 #include "Misc/Optional.h"
 #include "Templates/ValueOrError.h"
+#include "UObject/CoreReflectedTypeAccessors.h"
 
 class FText;
 template<typename NumericType> struct FNumericUnit;
@@ -96,6 +97,10 @@ enum class EUnit : uint8
 	Unspecified
 };
 
+#ifdef COREUOBJECT_API
+template<> COREUOBJECT_API UEnum* StaticEnum<EUnit>();
+#endif
+
 /** Enumeration that specifies particular classes of unit */
 enum class EUnitType
 {
diff --git a/Engine/Source/Runtime/Core/Public/Misc/AutomationEvent.h b/Engine/Source/Runtime/Core/Public/Misc/AutomationEvent.h
--- a/Engine/Source/Runtime/Core/Public/Misc/AutomationEvent.h
+++ b/Engine/Source/Runtime/Core/Public/Misc/AutomationEvent.h
@@ -5,6 +5,7 @@
 #include "CoreTypes.h"
 #include "Misc/DateTime.h"
 #include "Misc/Guid.h"
+#include "UObject/CoreReflectedTypeAccessors.h"
 
 enum class EAutomationEventType : uint8
 {
@@ -13,6 +14,10 @@ enum class EAutomationEventType : uint8
 	Error
 };
 
+#ifdef COREUOBJECT_API
+template<> COREUOBJECT_API UEnum* StaticEnum<EAutomationEventType>();
+#endif
+
 struct FAutomationEvent
 {
 public:
diff --git a/Engine/Source/Runtime/Core/Public/Misc/CString.h b/Engine/Source/Runtime/Core/Public/Misc/CString.h
--- a/Engine/Source/Runtime/Core/Public/Misc/CString.h
+++ b/Engine/Source/Runtime/Core/Public/Misc/CString.h
@@ -11,13 +11,14 @@
 #include "Templates/IsArrayOrRefOfTypeByPredicate.h"
 #include "Templates/IsValidVariadicFunctionArg.h"
 #include "Traits/IsCharEncodingCompatibleWith.h"
+#include "UObject/CoreReflectedTypeAccessors.h"
 
 #define MAX_SPRINTF 1024
 
 /** Determines case sensitivity options for string comparisons. */
 namespace ESearchCase
 {
-	enum Type
+	enum Type : int
 	{
 		/** Case sensitive. Upper/lower casing must match for strings to be considered equal. */
 		CaseSensitive,
@@ -27,10 +28,14 @@ namespace ESearchCase
 	};
 };
 
+#ifdef COREUOBJECT_API
+template<> COREUOBJECT_API UEnum* StaticEnum<ESearchCase::Type>();
+#endif
+
 /** Determines search direction for string operations. */
 namespace ESearchDir
 {
-	enum Type
+	enum Type : int
 	{
 		/** Search from the start, moving forward through the string. */
 		FromStart,
@@ -40,6 +45,10 @@ namespace ESearchDir
 	};
 }
 
+#ifdef COREUOBJECT_API
+template<> COREUOBJECT_API UEnum* StaticEnum<ESearchDir::Type>();
+#endif
+
 /** Helper class used to convert CString into a boolean value. */
 struct FToBoolHelper
 {
diff --git a/Engine/Source/Runtime/Core/Public/Misc/CoreMiscDefines.h b/Engine/Source/Runtime/Core/Public/Misc/CoreMiscDefines.h
--- a/Engine/Source/Runtime/Core/Public/Misc/CoreMiscDefines.h
+++ b/Engine/Source/Runtime/Core/Public/Misc/CoreMiscDefines.h
@@ -4,6 +4,7 @@
 
 #include "HAL/Platform.h"
 #include "HAL/PreprocessorHelpers.h"
+#include "UObject/CoreReflectedTypeAccessors.h"
 
 // When passed to pragma message will result in clickable warning in VS
 #define WARNING_LOCATION(Line) __FILE__ "(" UE_STRINGIZE(Line) ")"
@@ -603,6 +604,10 @@ enum class EInputDeviceConnectionState : uint8
 	Connected
 };
 
+#ifdef COREUOBJECT_API
+template<> COREUOBJECT_API UEnum* StaticEnum<EInputDeviceConnectionState>();
+#endif
+
 /** Data about an input device's current state */
 struct FPlatformInputDeviceState
 {
diff --git a/Engine/Source/Runtime/Core/Public/Misc/OutputDevice.h b/Engine/Source/Runtime/Core/Public/Misc/OutputDevice.h
--- a/Engine/Source/Runtime/Core/Public/Misc/OutputDevice.h
+++ b/Engine/Source/Runtime/Core/Public/Misc/OutputDevice.h
@@ -9,6 +9,7 @@
 #include "Templates/IsArrayOrRefOfTypeByPredicate.h"
 #include "Templates/IsValidVariadicFunctionArg.h"
 #include "Traits/IsCharEncodingCompatibleWith.h"
+#include "UObject/CoreReflectedTypeAccessors.h"
 
 class FString;
 class FText;
@@ -107,7 +108,7 @@ using namespace OutputDeviceColor;
  */
 namespace ELogTimes
 {
-	enum Type
+	enum Type : int
 	{
 		// Do not display log timestamps
 		None,
@@ -126,6 +127,10 @@ namespace ELogTimes
 	};
 }
 
+#ifdef COREUOBJECT_API
+template<> COREUOBJECT_API UEnum* StaticEnum<ELogTimes::Type>();
+#endif
+
 class FName;
 
 // An output device.
diff --git a/Engine/Source/Runtime/Core/Public/PixelFormat.h b/Engine/Source/Runtime/Core/Public/PixelFormat.h
--- a/Engine/Source/Runtime/Core/Public/PixelFormat.h
+++ b/Engine/Source/Runtime/Core/Public/PixelFormat.h
@@ -5,6 +5,7 @@
 
 #include "CoreTypes.h"
 #include "Misc/EnumClassFlags.h"
+#include "UObject/CoreReflectedTypeAccessors.h"
 
 #if defined(PF_MAX)
 #undef PF_MAX
@@ -109,8 +110,17 @@ enum EPixelFormat : uint8
 	PF_R16G16_SINT			=92,
 	PF_R8G8B8				=93,
 	PF_MAX					=94,
+
+	// If updating this enum, remember to update FOREACH_ENUM_EPIXELFORMAT_CORE below.
 };
-#define FOREACH_ENUM_EPIXELFORMAT(op) \
+
+#ifdef COREUOBJECT_API
+template<> COREUOBJECT_API UEnum* StaticEnum<EPixelFormat>();
+#endif
+
+// This should match the definition of FOREACH_ENUM_EPIXELFORMAT generated in NoExportTypes.generated.h.
+// It is defined separately here to prevent clashes of the macro name when NoExportTypes.h is also included.
+#define FOREACH_ENUM_EPIXELFORMAT_CORE(op) \
 	op(PF_Unknown) \
 	op(PF_A32B32G32R32F) \
 	op(PF_B8G8R8A8) \
@@ -203,8 +213,8 @@ enum EPixelFormat : uint8
 	op(PF_ASTC_8x8_NORM_RG) \
 	op(PF_ASTC_10x10_NORM_RG) \
 	op(PF_ASTC_12x12_NORM_RG) \
-	op(PF_R8G8B8)\
-	op(PF_R16G16_SINT)
+	op(PF_R16G16_SINT) \
+	op(PF_R8G8B8)
 
 // Defines which channel is valid for each pixel format
 enum class EPixelFormatChannelFlags : uint8
diff --git a/Engine/Source/Runtime/Core/Public/UObject/CoreReflectedTypeAccessors.h b/Engine/Source/Runtime/Core/Public/UObject/CoreReflectedTypeAccessors.h
new file mode 100644
--- /dev/null
+++ b/Engine/Source/Runtime/Core/Public/UObject/CoreReflectedTypeAccessors.h
@@ -0,0 +1,11 @@
+// Copyright Epic Games, Inc. All Rights Reserved.
+
+#pragma once
+
+class UEnum;
+
+/*-----------------------------------------------------------------------------
+C++ templated Static(Enum) retrieval function prototypes.
+-----------------------------------------------------------------------------*/
+
+template<typename EnumType> UEnum* StaticEnum();
diff --git a/Engine/Source/Runtime/CoreUObject/Public/UObject/NoExportTypes.h b/Engine/Source/Runtime/CoreUObject/Public/UObject/NoExportTypes.h
--- a/Engine/Source/Runtime/CoreUObject/Public/UObject/NoExportTypes.h
+++ b/Engine/Source/Runtime/CoreUObject/Public/UObject/NoExportTypes.h
@@ -63,6 +63,8 @@
 
 #endif
 
+#include "NoExportTypes.generated.h"
+
 #if !CPP      //noexport class
 
 /// @cond DOXYGEN_IGNORE
@@ -199,7 +201,7 @@ enum EInterpCurveMode : int
  * @warning: The *Tools DLLs will also need to be recompiled if the ordering is changed, but should not need code changes.
  */
 UENUM()
-enum EPixelFormat : int
+enum EPixelFormat : uint8
 {
 	PF_Unknown,
 	PF_A32B32G32R32F,
@@ -315,6 +317,7 @@ enum EPixelFormat : int
 	PF_ASTC_10x10_NORM_RG,
 	PF_ASTC_12x12_NORM_RG,
 	PF_R16G16_SINT,
+	PF_R8G8B8,
 	PF_MAX,
 };
 
diff --git a/Engine/Source/Runtime/CoreUObject/Public/UObject/ReflectedTypeAccessors.h b/Engine/Source/Runtime/CoreUObject/Public/UObject/ReflectedTypeAccessors.h
--- a/Engine/Source/Runtime/CoreUObject/Public/UObject/ReflectedTypeAccessors.h
+++ b/Engine/Source/Runtime/CoreUObject/Public/UObject/ReflectedTypeAccessors.h
@@ -2,12 +2,13 @@
 
 #pragma once
 
+#include "UObject/CoreReflectedTypeAccessors.h"
+
 class UClass;
 class UScriptStruct;
-class UEnum;
 
 /*-----------------------------------------------------------------------------
-C++ templated Static(Class/Struct/Enum) retrieval function prototypes.
+C++ templated Static(Class/Struct) retrieval function prototypes.
 -----------------------------------------------------------------------------*/
 
 template<typename ClassType> inline UClass* StaticClass()
@@ -19,5 +20,3 @@ template<typename StructType> inline UScriptStruct* StaticStruct()
 {
 	return StructType::StaticStruct();
 }
-
-template<typename EnumType> UEnum* StaticEnum();
diff --git a/Engine/Source/Runtime/CoreUObject/Public/UObject/UObjectGlobals.h b/Engine/Source/Runtime/CoreUObject/Public/UObject/UObjectGlobals.h
--- a/Engine/Source/Runtime/CoreUObject/Public/UObject/UObjectGlobals.h
+++ b/Engine/Source/Runtime/CoreUObject/Public/UObject/UObjectGlobals.h
@@ -38,6 +38,7 @@
 #include "Templates/IsTObjectPtr.h"
 #include "Traits/IsCharEncodingCompatibleWith.h"
 #include "UObject/CompiledInObjectPtr.h"
+#include "UObject/CoreReflectedTypeAccessors.h"
 #include "UObject/FindObjectFlags.h"
 #include "UObject/NameTypes.h"
 #include "UObject/ObjectMacros.h"
@@ -4259,6 +4260,8 @@ enum class EDataValidationResult : uint8
 	NotValidated
 };
 
+template<> COREUOBJECT_API UEnum* StaticEnum<EDataValidationResult>();
+
 /**
  * Combines two different data validation results and returns the combined result.
  *
