Clang under XCode 15 doe not properly implementing copy initialization in this context. Explicitely initialize FReadHandle/FWriteHandle.

#Synced-CL 49564828

diff --git a/Engine/Source/Runtime/Core/Public/HAL/PlatformProcess.h b/Engine/Source/Runtime/Core/Public/HAL/PlatformProcess.h
--- a/Engine/Source/Runtime/Core/Public/HAL/PlatformProcess.h
+++ b/Engine/Source/Runtime/Core/Public/HAL/PlatformProcess.h
@@ -72,7 +72,7 @@ namespace UE::HAL
 					ReadPipe = nullptr;
 					WritePipe = nullptr;
 				}
-				return TTuple<FReadHandle, FWriteHandle>{ReadPipe, WritePipe};
+				return TTuple<FReadHandle, FWriteHandle>{FReadHandle{ReadPipe}, FWriteHandle{WritePipe}};
 			}()}
 		{
 		}
