[Backout] - CL50321384 - CompactHashTable::Remove asserts in cook.
#rnx
Original CL Desc
-----------------------------------------------------------------
[Core] Compact set improvements
* Fix old naming from CompactHashTable to CompactHashTableView
* Rework remove function to do validations before modifying set for better debugging
* Fixup compact hash so it's properly validating lookup when relocating last element
* Normalize printing functions to use %d for num elements since we only support 32 bit for now
* Move out of memory check to cover all cases, and fixup some raw (u)int32 usages to use appropriate typedefs

#tests Used to help track down issues with improper set usage, ran many preflights

#Synced-CL 50170368

#Synced-CL 50358454

diff --git a/Engine/Source/Runtime/Core/Public/Containers/CompactHashTable.h b/Engine/Source/Runtime/Core/Public/Containers/CompactHashTable.h
--- a/Engine/Source/Runtime/Core/Public/Containers/CompactHashTable.h
+++ b/Engine/Source/Runtime/Core/Public/Containers/CompactHashTable.h
@@ -30,9 +30,9 @@
 
 namespace UE::Core::CompactHashTable::Private
 {
-	// The reference to the given index in chain of a given hash key and return a pointer to it
+	// Remove an element from the list and patch up any next index references, after this the spot is empty and needs to be filled
 	template<typename IndexType>
-	inline IndexType* FindLookupEntry(const uint32 Index, const uint32 Key, IndexType* HashData, const uint32 HashCount, IndexType* NextIndexData, const uint32 NextIndexCount)
+	inline IndexType RemoveInternal(const uint32 Index, const uint32 Key, IndexType* HashData, const uint32 HashCount, IndexType* NextIndexData, const uint32 NextIndexCount)
 	{
 		const uint32 HashIndex = Key & (HashCount - 1);
 
@@ -40,12 +40,14 @@ namespace UE::Core::CompactHashTable::Private
 		{
 			if (*Lookup == Index)
 			{
-				return Lookup;
+				// Next = Next->Next
+				*Lookup = NextIndexData[Index];
+				return *Lookup;
 			}
 		}
 
-		checkf(false, TEXT("Failed to find entry from set.")); // This usually indicates that something changed in the entry that affected the hash so it no longer matches up.
-		return nullptr; // will crash
+		checkf(false, TEXT("Failed to find object to remove from set.")); // This usually indicates that something changed in the entry that affected the hash so it no longer matches up.
+		return (IndexType)INDEX_NONE;
 	}
 }
 
@@ -162,22 +164,28 @@ namespace UE::Core::CompactHashTable
 	{
 		checkSlow(LastIndex < NextIndexCount && Index <= LastIndex);
 
-		// Next = Next->Next
-		IndexType* LookupEntry = UE::Core::CompactHashTable::Private::FindLookupEntry<IndexType>(Index, Key, HashData, HashCount, NextIndexData, NextIndexCount);
-		*LookupEntry = NextIndexData[Index];
+		IndexType NextIndex = UE::Core::CompactHashTable::Private::RemoveInternal<IndexType>(Index, Key, HashData, HashCount, NextIndexData, NextIndexCount);
 
 		if (Index != LastIndex)
 		{
 			// Remap the last element into the now empty spot
-			IndexType* LastLookupEntry = UE::Core::CompactHashTable::Private::FindLookupEntry<IndexType>(LastIndex, OptLastKey, HashData, HashCount, NextIndexData, NextIndexCount);
-			*LastLookupEntry = (IndexType)Index;
+			const uint32 HashIndex = OptLastKey & (HashCount - 1);
 
-			// Move the last elements next index into the new slot as well
-			NextIndexData[Index] = NextIndexData[LastIndex];
+			for (IndexType* Lookup = &HashData[HashIndex]; *Lookup < NextIndexCount; Lookup = &NextIndexData[*Lookup])
+			{
+				if (*Lookup == LastIndex)
+				{
+					// Fixup to look at now empty slot
+					*Lookup = (IndexType)Index;
+					// Now empty slot looks where the last item looked
+					NextIndexData[Index] = NextIndexData[LastIndex];
+					break;
+				}
+			}
 		}
 
-		// Return whatever next is at the original lookup entry that we fixed up. If it happened to be the last index the the previous block would have fixed it up for us
-		return IndexTypeToSizeType(*LookupEntry);
+		// If the next index points to the last element, then it will get relocated into the current element so return the current instead of the next
+		return NextIndex == LastIndex ? Index : IndexTypeToSizeType(NextIndex);
 	}
 
 	/* Remove an element from the list, shift all indexes down to preserve the order of elements in a list
@@ -199,12 +207,8 @@ namespace UE::Core::CompactHashTable
 			}
 		};
 
-		IndexType* LookupEntry = UE::Core::CompactHashTable::Private::FindLookupEntry<IndexType>(Index, Key, HashData, HashCount, NextIndexData, NextIndexCount);
-		*LookupEntry = NextIndexData[Index];
-		FixIndex(*LookupEntry);
-
-		// Cache return value here since the memory location could move if it was a 'next index' that's located after where Index is
-		const uint32 NextIndex = IndexTypeToSizeType(*LookupEntry);
+		IndexType NextIndex = UE::Core::CompactHashTable::Private::RemoveInternal<IndexType>(Index, Key, HashData, HashCount, NextIndexData, NextIndexCount);
+		FixIndex(NextIndex);
 
 		// For the hash indexes, just decrement any that are bigger than the removed element
 		for (uint32 HashIndex = 0; HashIndex < HashCount; ++HashIndex)
@@ -232,7 +236,7 @@ namespace UE::Core::CompactHashTable
 			}
 		}
 
-		return NextIndex;
+		return IndexTypeToSizeType(NextIndex);
 	}
 }
 
diff --git a/Engine/Source/Runtime/Core/Public/Containers/CompactSet.h.inl b/Engine/Source/Runtime/Core/Public/Containers/CompactSet.h.inl
--- a/Engine/Source/Runtime/Core/Public/Containers/CompactSet.h.inl
+++ b/Engine/Source/Runtime/Core/Public/Containers/CompactSet.h.inl
@@ -614,7 +614,7 @@ public:
 		// Template property, branch will be optimized out
 		if constexpr (AllocatorType::RequireRangeCheck)
 		{
-			checkf(IsValidId(Id), TEXT("Array index out of bounds: %d into an array of size %d"), Id.AsInteger(), this->NumElements);
+			checkf(IsValidId(Id), TEXT("Array index out of bounds: %d into an array of size %lld"), Id.AsInteger(), (long long)this->NumElements); // & for one branch
 		}
 	}
 
@@ -1147,13 +1147,13 @@ public:
 	{
 		if (this->MaxElements == 0)
 		{
-			Ar.Log(TEXT("UE_TCOMPACT_SET: empty"));
+			Ar.Logf(TEXT("UE_TCOMPACT_SET: empty"), this->NumElements, this->MaxElements);
 		}
 		else if (this->MaxElements > 0)
 		{
 			FConstCompactHashTableView HashTable = this->GetConstHashTableView();
 
-			Ar.Logf(TEXT("UE_TCOMPACT_SET: %d elements, %d max elements, %d hash slots"), this->NumElements, this->MaxElements, HashTable.GetHashCount());
+			Ar.Logf(TEXT("UE_TCOMPACT_SET: %i elements, %i max elements, %i hash slots"), this->NumElements, this->MaxElements, HashTable.GetHashCount());
 
 			for (uint32 HashIndex = 0; HashIndex < HashTable.GetHashCount(); ++HashIndex)
 			{
@@ -1495,45 +1495,42 @@ private:
 	template <bool IsStable = false>
 	void RemoveByIndex(const SizeType ElementIndex)
 	{
-		checkf(ElementIndex >= 0 && ElementIndex < this->NumElements, TEXT("Invalid ElementIndex (%d) passed to UE_TCOMPACT_SET::RemoveByIndex. NumElements=%d"), ElementIndex, this->NumElements);
+		checkf(ElementIndex >= 0 && ElementIndex < this->NumElements, TEXT("Invalid ElementIndex passed to UE_TCOMPACT_SET::RemoveByIndex"));
+		RemoveByIndexAndHash<IsStable>(ElementIndex, KeyFuncs::GetKeyHash(KeyFuncs::GetSetKey(GetData()[ElementIndex])));
+	}
+
+	template <bool IsStable = false>
+	SizeType RemoveByIndexAndHash(const SizeType ElementIndex, const uint32 KeyHash)
+	{
+		checkf(ElementIndex >= 0 && ElementIndex < this->NumElements, TEXT("Invalid ElementIndex passed to UE_TCOMPACT_SET::RemoveByIndex"));
 
 		ElementType* ElementsData = GetData();
 		FCompactHashTableView HashTable = GetHashTableView();
 
-		RemoveByIndexAndHash<IsStable>(ElementIndex, KeyFuncs::GetKeyHash(KeyFuncs::GetSetKey(ElementsData[ElementIndex])), ElementsData, HashTable);
-	}
+		--this->NumElements;
 
-	template <bool IsStable = false>
-	SizeType RemoveByIndexAndHash(const SizeType ElementIndex, const uint32 KeyHash, ElementType* ElementsData, FCompactHashTableView& HashTable)
-	{
-		const SizeType LastElementIndex = this->NumElements - 1;
-		SizeType NextIndex; 
-		
-		// Note: Destructing/moving elements for better debugability
+		const SizeType LastElementIndex = this->NumElements;
 		if (ElementIndex == LastElementIndex)
 		{
-			NextIndex = HashTable.Remove(ElementIndex, KeyHash, ElementIndex, KeyHash);
-			
 			ElementsData[LastElementIndex].~ElementType();
+			return HashTable.Remove(ElementIndex, KeyHash, ElementIndex, KeyHash);
 		}
-		else if constexpr (IsStable)
-		{
-			NextIndex = HashTable.RemoveStable(ElementIndex, KeyHash);
 
+		if constexpr (IsStable)
+		{
 			ElementsData[ElementIndex].~ElementType();
 			RelocateConstructItems<ElementType>(ElementsData + ElementIndex, ElementsData + ElementIndex + 1, LastElementIndex - ElementIndex);
+
+			return HashTable.RemoveStable(ElementIndex, KeyHash);
 		}
 		else
 		{
 			const uint32 LastElementHash = KeyFuncs::GetKeyHash(KeyFuncs::GetSetKey(ElementsData[LastElementIndex]));
-			NextIndex = HashTable.Remove(ElementIndex, KeyHash, LastElementIndex, LastElementHash);
-			
-			MoveByRelocate(ElementsData[ElementIndex], ElementsData[LastElementIndex]);
-		}
 
-		--this->NumElements;
+			MoveByRelocate(ElementsData[ElementIndex], ElementsData[LastElementIndex]);
 
-		return NextIndex;
+			return HashTable.Remove(ElementIndex, KeyHash, LastElementIndex, LastElementHash);
+		}
 	}
 
 	template<typename ComparableKey, bool IsStable = false>
@@ -1542,25 +1539,24 @@ private:
 		checkSlow(this->NumElements > 0);
 		int32 NumRemovedElements = 0;
 
-		ElementType* ElementsData = GetData();
+		const ElementType* ElementsData = GetData();
 		FCompactHashTableView HashTable = GetHashTableView();
 
 		SizeType ElementIndex = HashTable.GetFirst(KeyHash);
 
 		while (ElementIndex != INDEX_NONE)
 		{
-			checkf(ElementIndex >= 0 && ElementIndex < this->NumElements, TEXT("Invalid ElementIndex (%d) passed to UE_TCOMPACT_SET::RemoveByIndexAndHash. NumElements=%d"), ElementIndex, this->NumElements);
 			if (KeyFuncs::Matches(KeyFuncs::GetSetKey(ElementsData[ElementIndex]), Key))
 			{
 				if constexpr (!KeyFuncs::bAllowDuplicateKeys)
 				{
-					RemoveByIndexAndHash<IsStable>(ElementIndex, KeyHash, ElementsData, HashTable);
+					RemoveByIndexAndHash<IsStable>(ElementIndex, KeyHash);
 					// If the hash disallows duplicate keys, we're done removing after the first matched key.
 					return 1;
 				}
 				else
 				{
-					ElementIndex = RemoveByIndexAndHash<IsStable>(ElementIndex, KeyHash, ElementsData, HashTable);
+					ElementIndex = RemoveByIndexAndHash<IsStable>(ElementIndex, KeyHash);
 					NumRemovedElements++;
 				}
 			}
diff --git a/Engine/Source/Runtime/Core/Public/Containers/CompactSetBase.h b/Engine/Source/Runtime/Core/Public/Containers/CompactSetBase.h
--- a/Engine/Source/Runtime/Core/Public/Containers/CompactSetBase.h
+++ b/Engine/Source/Runtime/Core/Public/Containers/CompactSetBase.h
@@ -38,19 +38,19 @@ public:
 	}
 
 	/** @return the number of elements. */
-	[[nodiscard]] UE_FORCEINLINE_HINT SizeType Num() const
+	[[nodiscard]] UE_FORCEINLINE_HINT int32 Num() const
 	{
 		return NumElements;
 	}
 
 	/** @return The number of elements the set can hold before reallocation. */
-	[[nodiscard]] UE_FORCEINLINE_HINT SizeType Max() const
+	[[nodiscard]] UE_FORCEINLINE_HINT int32 Max() const
 	{
 		return MaxElements;
 	}
 
 	/** Deprecated - unnecessary, keeping this here so TCompactSet can be swapped with TSet without changing code */
-	[[nodiscard]] UE_FORCEINLINE_HINT SizeType GetMaxIndex() const
+	[[nodiscard]] UE_FORCEINLINE_HINT int32 GetMaxIndex() const
 	{
 		return NumElements;
 	}
@@ -149,9 +149,9 @@ protected:
 		return FMath::Min<SizeType>(MaxElements, RealAvailableBytes / Layout.Size); // Now we can get the true number of elements we could potentially use within the aligned space
 	}
 
-	[[nodiscard]] SizeType AllocatorCalculateSlackGrow(SizeType NewMaxElements, const FCompactSetLayout& Layout) const
+	[[nodiscard]] int32 AllocatorCalculateSlackGrow(int32 NewMaxElements, const FCompactSetLayout& Layout) const
 	{
-		const size_t OldHashCount = MaxElements > 0 ? *(HashCountType*)(Elements.GetAllocation() + GetElementsSizeInBytes(MaxElements, Layout)) : 0;
+		const size_t OldHashCount = MaxElements > 0 ? *(uint32*)(Elements.GetAllocation() + GetElementsSizeInBytes(MaxElements, Layout)) : 0;
 		const size_t OldSize = MaxElements > 0 ? GetTotalMemoryRequiredInBytes(MaxElements, OldHashCount, Layout) : 0;
 
 		const size_t NewHashCount = NewMaxElements > 0 ? GetHashCount(NewMaxElements) : 0;
@@ -196,28 +196,28 @@ protected:
 		return SlackNumElements;
 	}
 
-	void ResizeAllocation(const SizeType NewMaxElements, const FCompactSetLayout& Layout)
+	void ResizeAllocation(const int32 NewMaxElements, const FCompactSetLayout& Layout)
 	{
 		(void)ResizeAllocationPreserveData(NewMaxElements, Layout, false);
 	}
 
 	// Use this is you'll be keeping the element data
-	[[nodiscard]] bool ResizeAllocationPreserveData(const SizeType NewMaxElements, const FCompactSetLayout& Layout, bool bPreserve = true)
+	[[nodiscard]] bool ResizeAllocationPreserveData(const int32 NewMaxElements, const FCompactSetLayout& Layout, bool bPreserve = true)
 	{
 		bool bRequiresRehash = false;
 
 		if (NewMaxElements != MaxElements)
 		{
-			const size_t OldHashCount = MaxElements > 0 ? *(HashCountType*)(Elements.GetAllocation() + GetElementsSizeInBytes(MaxElements, Layout)) : 0;
+			const size_t OldHashCount = MaxElements > 0 ? *(uint32*)(Elements.GetAllocation() + GetElementsSizeInBytes(MaxElements, Layout)) : 0;
 			const size_t OldSize = MaxElements > 0 ? GetTotalMemoryRequiredInBytes(MaxElements, OldHashCount, Layout) : 0;
 
 			const size_t NewHashCount = NewMaxElements > 0 ? UE::Core::CompactHashTable::GetHashCount(NewMaxElements) : 0;
 			const size_t NewSize = NewMaxElements > 0 ? GetTotalMemoryRequiredInBytes(NewMaxElements, NewHashCount, Layout) : 0;
 
-			checkf(NewSize >= 0 && NewSize <= MAX_int32, TEXT("Invalid size for TSet[%lld]: NewMaxElements[%d] ElementSize[%d] HashCount[%lld]"), NewSize, NewMaxElements, Layout.Size, NewHashCount);
-
 			if (bPreserve && NewMaxElements > MaxElements && OldHashCount == NewHashCount)
 			{
+				checkf(NewSize >= 0 && NewSize <= MAX_int32, TEXT("Invalid size for TSet[%d]: NewMaxElements[%d] ElementSize[%d] HashCount[%d]"), NewSize, NewMaxElements, Layout.Size, NewHashCount);
+
 				// If preserving then we copy over all the data
 				if constexpr (TAllocatorTraits<AllocatorType>::SupportsElementAlignment)
 				{
@@ -260,7 +260,7 @@ protected:
 				// Update hash size
 				if (NewMaxElements > 0)
 				{
-					HashCountType* NewHashTable = (HashCountType*)(Elements.GetAllocation() + GetElementsSizeInBytes(NewMaxElements, Layout));
+					uint32* NewHashTable = (uint32*)(Elements.GetAllocation() + GetElementsSizeInBytes(NewMaxElements, Layout));
 
 					*NewHashTable = NewHashCount;
 					bRequiresRehash = true;
diff --git a/Engine/Source/Runtime/Core/Tests/Containers/CompactHashTableTests.cpp b/Engine/Source/Runtime/Core/Tests/Containers/CompactHashTableTests.cpp
--- a/Engine/Source/Runtime/Core/Tests/Containers/CompactHashTableTests.cpp
+++ b/Engine/Source/Runtime/Core/Tests/Containers/CompactHashTableTests.cpp
@@ -188,7 +188,7 @@ TEST_CASE_NAMED(FStaticCompactHashTableTest, "System::Core::Containers::CompactH
 	}
 }
 
-TEST_CASE_NAMED(FCompactHashTableViewTest, "System::Core::Containers::CompactHashTable::Dynamic", "[Core][Containers][CompactHashTable]")
+TEST_CASE_NAMED(FDynamicCompactHashTableTest, "System::Core::Containers::CompactHashTable::Dynamic", "[Core][Containers][CompactHashTable]")
 {
 	SECTION("Constructor Behavior")
 	{
