Remove unused bPrecacheAsSoonAsPossible
#rnx

#Synced-CL 50326944

diff --git a/Engine/Source/Runtime/Core/Private/HAL/FileManagerGeneric.cpp b/Engine/Source/Runtime/Core/Private/HAL/FileManagerGeneric.cpp
--- a/Engine/Source/Runtime/Core/Private/HAL/FileManagerGeneric.cpp
+++ b/Engine/Source/Runtime/Core/Private/HAL/FileManagerGeneric.cpp
@@ -786,41 +786,8 @@ bool FArchiveFileReaderGeneric::InternalPrecache( int64 PrecacheOffset, int64 Pr
 	int64 WriteOffset = 0;
 	if (bPosWithinBuffer)
 	{
-		if (bPrecacheAsSoonAsPossible)
-		{
-			// If we already have some bytes after pos buffered, check whether it is sufficient to satisfy this precache request.
-			int64 RemainingBufferCount = BufferArray.Num() - (Pos - BufferBase);
-			int64 RequestedBufferCount = FMath::Min(PrecacheSize, BufferCount);
-			if (RemainingBufferCount >= RequestedBufferCount)
-			{
-				// The requested precache range is already in our buffer.
-				return true;
-			}
-
-			// We need to read more to satisfy the precache request.
-			// Since Pos is within the buffer, the low-level read position is at the end of the buffer. 
-			// Copy the existing bytes after Pos out of the old buffer into the new buffer, and then read only the remaining bytes from the low-level handle into the rest of the new buffer.
-			if (BufferCount <= BufferSize)
-			{
-				// We don't need to reallocate the buffer, so we can just move the RemainingBufferCount bytes from the end of the buffer to the beginning.
-				FMemory::Memmove(BufferArray.GetData(), BufferArray.GetData() + Pos - BufferBase, RemainingBufferCount);
-				BufferArray.SetNumUninitialized(BufferCount, EAllowShrinking::No);
-			}
-			else
-			{
-				TArray64<uint8> OldArray(MoveTemp(BufferArray));
-				BufferArray.SetNumUninitialized(BufferCount, EAllowShrinking::No);
-				FMemory::Memcpy(BufferArray.GetData(), OldArray.GetData() + Pos - BufferBase, RemainingBufferCount);
-			}
-
-			ReadCount = BufferCount - RemainingBufferCount;
-			WriteOffset = RemainingBufferCount;
-		}
-		else
-		{
-			// At least one byte after the requested PrecacheOffset is in our existing buffer (since bPosWithinBuffer is true), so do nothing and return true
-			return true;
-		}
+		// At least one byte after the requested PrecacheOffset is in our existing buffer (since bPosWithinBuffer is true), so do nothing and return true
+		return true;
 	}
 	else
 	{
diff --git a/Engine/Source/Runtime/Core/Public/HAL/FileManagerGeneric.h b/Engine/Source/Runtime/Core/Public/HAL/FileManagerGeneric.h
--- a/Engine/Source/Runtime/Core/Public/HAL/FileManagerGeneric.h
+++ b/Engine/Source/Runtime/Core/Public/HAL/FileManagerGeneric.h
@@ -288,11 +288,6 @@ protected:
 	uint32 Flags;
 	bool bFirstReadAfterSeek;
 
-	enum
-	{
-		bPrecacheAsSoonAsPossible = 0 // Setting this to true makes it more likely bytes will be available without waiting due to external Precache calls, but at the cost of a larger number of read requests.
-	};
-
 	friend class FArchiveFileReaderGenericTest;
 };
 
diff --git a/Engine/Source/Runtime/Core/Tests/HAL/FileManagerGenericTest.cpp b/Engine/Source/Runtime/Core/Tests/HAL/FileManagerGenericTest.cpp
--- a/Engine/Source/Runtime/Core/Tests/HAL/FileManagerGenericTest.cpp
+++ b/Engine/Source/Runtime/Core/Tests/HAL/FileManagerGenericTest.cpp
@@ -76,14 +76,7 @@ void FArchiveFileReaderGenericTest::TestInternalPrecache()
 	SetPosAndBuffer(Reader, PosStart, 0, BufferSize);
 	Result = Reader->InternalPrecache(PosStart, BufferSize);
 	CHECK_MESSAGE(TEXT("PrecachePartwayPosMoreThanBuffer - Should succeed"), Result);
-	if (FArchiveFileReaderGeneric::bPrecacheAsSoonAsPossible)
-	{
-		CHECK_MESSAGE(TEXT("PrecachePartwayPosMoreThanBuffer - BufferBase should be updated"), Reader->BufferBase == PosStart);
-	}
-	else
-	{
-		CHECK_MESSAGE(TEXT("PrecachePartwayPosMoreThanBuffer - BufferBase should not be updated"), Reader->BufferBase == 0LL);
-	}
+	CHECK_MESSAGE(TEXT("PrecachePartwayPosMoreThanBuffer - BufferBase should not be updated"), Reader->BufferBase == 0LL);
 	CHECK_MESSAGE(TEXT("PrecachePartwayPosMoreThanBuffer - BufferCount should be set to BufferSize"), Reader->BufferArray.Num() == BufferSize);
 	TestBytesValid(TEXT("PrecachePartwayPosMoreThanBuffer - BufferBytes should match BufferBase"), Reader->BufferArray.GetData(), Reader->BufferArray.Num(), Reader->BufferBase);
 
@@ -92,16 +85,8 @@ void FArchiveFileReaderGenericTest::TestInternalPrecache()
 	Reader->BufferArray.Shrink();
 	Result = Reader->InternalPrecache(3*BufferSize/4, BufferSize);
 	CHECK_MESSAGE(TEXT("PrecachePartwayPosMoreThanBufferReallocation - Should succeed"), Result);
-	if (FArchiveFileReaderGeneric::bPrecacheAsSoonAsPossible)
-	{
-		CHECK_MESSAGE(TEXT("PrecachePartwayPosMoreThanBufferReallocation - BufferBase should be updated"), Reader->BufferBase == 3*BufferSize/4);
-		CHECK_MESSAGE(TEXT("PrecachePartwayPosMoreThanBufferReallocation - BufferCount should be set to BufferSize"), Reader->BufferArray.Num() == BufferSize);
-	}
-	else
-	{
-		CHECK_MESSAGE(TEXT("PrecachePartwayPosMoreThanBufferReallocation - BufferBase should not be updated"), Reader->BufferBase == BufferSize/2);
-		CHECK_MESSAGE(TEXT("PrecachePartwayPosMoreThanBufferReallocation - BufferCount should not be updated"), Reader->BufferArray.Num() == BufferSize/2);
-	}
+	CHECK_MESSAGE(TEXT("PrecachePartwayPosMoreThanBufferReallocation - BufferBase should not be updated"), Reader->BufferBase == BufferSize/2);
+	CHECK_MESSAGE(TEXT("PrecachePartwayPosMoreThanBufferReallocation - BufferCount should not be updated"), Reader->BufferArray.Num() == BufferSize/2);
 	TestBytesValid(TEXT("PrecachePartwayPosMoreThanBufferReallocation - BufferBytes should match BufferBase"), Reader->BufferArray.GetData(), Reader->BufferArray.Num(), Reader->BufferBase);
 
 	// InternalPrecache partway through the buffer, when the requested size is within the remaining buffer, should not change the buffer
@@ -144,16 +129,8 @@ void FArchiveFileReaderGenericTest::TestInternalPrecache()
 	Reader->SeekLowLevel(FileSize);
 	Result = Reader->InternalPrecache(8, 16);
 	CHECK_MESSAGE(TEXT("PrecacheReadFailsBytesRemain - Should succeed if bytes are left in the buffer"), Result);
-	if (FArchiveFileReaderGeneric::bPrecacheAsSoonAsPossible)
-	{
-		CHECK_MESSAGE(TEXT("PrecacheReadFailsBytesRemain - BufferBase should be updated"), Reader->BufferBase == 8LL);
-		CHECK_MESSAGE(TEXT("PrecacheReadFailsBytesRemain - BufferCount should be updated"), Reader->BufferArray.Num() == 8LL);
-	}
-	else
-	{
-		CHECK_MESSAGE(TEXT("PrecacheReadFailsBytesRemain - BufferBase should not be updated"), Reader->BufferBase == 0LL);
-		CHECK_MESSAGE(TEXT("PrecacheReadFailsBytesRemain - BufferCount should not be updated"), Reader->BufferArray.Num() == 16LL);
-	}
+	CHECK_MESSAGE(TEXT("PrecacheReadFailsBytesRemain - BufferBase should not be updated"), Reader->BufferBase == 0LL);
+	CHECK_MESSAGE(TEXT("PrecacheReadFailsBytesRemain - BufferCount should not be updated"), Reader->BufferArray.Num() == 16LL);
 
 	// If ReadLowLevel fails, InternalPrecache should return false if the first byte of PrecacheOffset is not in the buffer
 	SetPosAndBuffer(Reader, 16, 0, 16);
