Adding in support to generate CodeOrContent LLM tag values in Performance + LLM reports when enabled.
#tests Local tests on cooked builds, replay run on Horde and builds.

#Synced-CL 50149305

#Synced-CL 50149305

diff --git a/Engine/Source/Runtime/Core/Private/HAL/LLM/LLM.cpp b/Engine/Source/Runtime/Core/Private/HAL/LLM/LLM.cpp
--- a/Engine/Source/Runtime/Core/Private/HAL/LLM/LLM.cpp
+++ b/Engine/Source/Runtime/Core/Private/HAL/LLM/LLM.cpp
@@ -949,6 +949,7 @@ UE_TRACE_EVENT_END()
 #if LLM_CSV_PROFILER_WRITER_ENABLED
 	CSV_DEFINE_CATEGORY(LLM, true);
 	CSV_DEFINE_CATEGORY(LLMPlatform, true);
+	CSV_DEFINE_CATEGORY(LLMCodeOrContent, true);
 #endif
 
 TAutoConsoleVariable<int32> CVarLLMTrackPeaks(TEXT("LLM.TrackPeaks"), 0,
@@ -4849,7 +4850,7 @@ int32 FTagData::GetIndex() const
 
 bool FTagData::IsReportable() const
 {
-	return bIsReportable && TagSet == ELLMTagSet::None;
+	return bIsReportable && (TagSet == ELLMTagSet::None || TagSet == ELLMTagSet::CodeOrContent);
 }
 
 bool FTagData::IsStatsReportable() const
@@ -6871,14 +6872,75 @@ void FLLMThreadState::ClearAllocTypeAmounts()
 }
 
 // FLLMCsvWriter implementation.
-bool FLLMCsvWriter::FColumnKey::operator==(const FLLMCsvWriter::FColumnKey& Other) const
+bool CsvWriter::FColumnKey::operator==(const CsvWriter::FColumnKey& Other) const
 {
 	return SystemsTag == Other.SystemsTag && CodeOrContentTag == Other.CodeOrContentTag;
 }
 
-uint32 GetTypeHash(const FLLMCsvWriter::FColumnKey& Key)
+// Updates the Columns tracking CodeOrContent with their System and CodeOrContent tags as well as calculating their sizes.
+// Returns true if the columns were updated.
+static bool CalculateCodeOrContentColumns(TLLMMap<CsvWriter::FColumnKey, CsvWriter::FColumnData>& Columns, const ELLMTracker& Tracker,
+	FLLMGlobals& LLMRef, const FMapActiveTagsToAllocationGroup* ActiveTagsToAllocationGroup, 
+	UE::FPlatformSharedMutex& ReadLockMutex)
 {
-	return ::GetTypeHash(Key.SystemsTag) * 101 + ::GetTypeHash(Key.CodeOrContentTag);
+	// When recording ELLMTagSets::CodeOrContent, we write a column for every combination of ELLMTagSet::None
+	// and ELLMTagSet::CodeOrContent, rather than a column for every TagData. To find the sizes of the 
+	// combinations, we have to iterate over AllocationGroup sizes rather than TagData sizes.
+	// One extra bit of complexity: it is possible we are recording other TagSets as well, in which case
+	// each AllocationGroup tracks a more specific combination:
+	// { ELLMTagSet::None x ELLMTagSet::CodeOrContent x ELLMTagSet::<Other> }
+	// To handle that case, we have to sum up the size of every AllocationGroup that matches the 
+	// { ELLMTagSet::None x ELLMTagSet::CodeOrContent }
+	// combination we are interested in.
+	for (TPair<CsvWriter::FColumnKey, CsvWriter::FColumnData>& Pair : Columns)
+	{
+		Pair.Value.Size = 0;
+	}
+
+	int32 IndexOfCodeOrContentInActiveTags = LLMRef.GetIndexInFActiveTags(ELLMTagSet::CodeOrContent);
+	// We are recording CodeOrContent, so it should have an index.
+	LLMCheck(IndexOfCodeOrContentInActiveTags >= 0);
+	bool bUpdatedColumns = false;
+
+	UE::TSharedLock AllocationsReadLock(ReadLockMutex);
+	for (const TPair<FActiveTags, FAllocationGroup*>& GroupPair : (*ActiveTagsToAllocationGroup))
+	{
+		const FActiveTags& ActiveTags = GroupPair.Key;
+		const FTagData* SystemsTag = ActiveTags.GetSystemsTagData();
+		if (!SystemsTag->IsReportable())
+		{
+			continue;
+		}
+		if (SystemsTag->GetDisplayName() == TagName_Untagged)
+		{
+			// We write the "Untagged" tag in the 2nd column, and always write it, rather than writing it in the
+			// column index where we first notice it. That writing is handled by use of OverrideUntaggedTagData in
+			// WriteHeader and AddRow.
+			continue;
+		}
+		const FTagData* CodeOrContentTag = ActiveTags[IndexOfCodeOrContentInActiveTags];
+		if (!CodeOrContentTag)
+		{
+			continue;
+		}
+		int64 CurrentSize = GroupPair.Value->GetSize(Tracker);
+		// AllocationGroups includes some groups used by other Trackers and not by this one. We have to check whether
+		// the AllocationGroup has a non-zero size in our Tracker before adding it to our output. 
+		if (CurrentSize == 0)
+		{
+			continue;
+		}
+
+		CsvWriter::FColumnData& ColumnData = Columns.FindOrAdd(CsvWriter::FColumnKey{ SystemsTag, CodeOrContentTag });
+		if (!ColumnData.bInitialized)
+		{
+			ColumnData.bInitialized = true;
+			bUpdatedColumns = true;
+		}
+		ColumnData.Size += CurrentSize;
+	}
+
+	return bUpdatedColumns;
 }
 
 FLLMCsvWriter::FLLMCsvWriter()
@@ -7043,7 +7105,7 @@ bool FLLMCsvWriter::CreateArchive(FLLMGlobals& LLMRef)
 				int32 SizeForTag = TagNameStr.Len() + 1;
 				if (bRecordingCodeOrContent)
 				{
-					SizeForTag = SizeForTag + 1 + CodeOrContentCodeLength +
+					SizeForTag += SizeForTag + 1 + CodeOrContentCodeLength +
 						SizeForTag + 1 + CodeOrContentContentLength;
 				}
 				HeaderMaxSize += SizeForTag;
@@ -7063,87 +7125,32 @@ bool FLLMCsvWriter::UpdateColumns(FLLMGlobals& LLMRef, const FTrackerTagSizeMap&
 
 	if (bRecordingCodeOrContent)
 	{
-		// When recording ELLMTagSets::CodeOrContent, we write a column for every combination of ELLMTagSet::None
-		// and ELLMTagSet::CodeOrContent, rather than a column for every TagData. To find the sizes of the 
-		// combinations, we have to iterate over AllocationGroup sizes rather than TagData sizes.
-		// One extra bit of complexity: it is possible we are recording other TagSets as well, in which case
-		// each AllocationGroup tracks a more specific combination:
-		// { ELLMTagSet::None x ELLMTagSet::CodeOrContent x ELLMTagSet::<Other> }
-		// To handle that case, we have to sum up the size of every AllocationGroup that matches the 
-		// { ELLMTagSet::None x ELLMTagSet::CodeOrContent }
-		// combination we are interested in.
-		for (TPair<FColumnKey, FColumnData>& Pair : Columns)
-		{
-			Pair.Value.Size = 0;
-		}
-
-		int32 IndexOfCodeOrContentInActiveTags = LLMRef.GetIndexInFActiveTags(ELLMTagSet::CodeOrContent);
-		// We are recording CodeOrContent, so it should have an index.
-		LLMCheck(IndexOfCodeOrContentInActiveTags >= 0);
+		bUpdated = CalculateCodeOrContentColumns(Columns, Tracker, LLMRef, LLMRef.ActiveTagsToAllocationGroup, LLMRef.AllocationGroupsMutex);
+	}
 
-		UE::TSharedLock AllocationsReadLock(LLMRef.AllocationGroupsMutex);
-		for (const TPair<FActiveTags, FAllocationGroup*>& GroupPair : (*LLMRef.ActiveTagsToAllocationGroup))
+	for (const TPair<const FTagData*, FTrackerTagSizeData>& It : TagSizes)
+	{
+		const FTagData* SystemsTag = It.Key;
+		if (SystemsTag->GetTagSet() != ELLMTagSet::None || !SystemsTag->IsReportable())
 		{
-			const FActiveTags& ActiveTags = GroupPair.Key;
-			const FTagData* SystemsTag = ActiveTags.GetSystemsTagData();
-			if (!SystemsTag->IsReportable())
-			{
-				continue;
-			}
-			if (SystemsTag->GetDisplayName() == TagName_Untagged)
-			{
-				// We write the "Untagged" tag in the 2nd column, and always write it, rather than writing it in the
-				// column index where we first notice it. That writing is handled by use of OverrideUntaggedTagData in
-				// WriteHeader and AddRow.
-				continue; 
-			}
-			const FTagData* CodeOrContentTag = ActiveTags[IndexOfCodeOrContentInActiveTags];
-			if (!CodeOrContentTag)
-			{
-				continue;
-			}
-			int64 CurrentSize = GroupPair.Value->GetSize(this->Tracker);
-			// AllocationGroups includes some groups used by other Trackers and not by this one. We have to check whether
-			// the AllocationGroup has a non-zero size in our Tracker before adding it to our output. 
-			if (CurrentSize == 0)
-			{
-				continue;
-			}
-			FColumnData& ColumnData = Columns.FindOrAdd(FColumnKey{ SystemsTag, CodeOrContentTag });
-			if (!ColumnData.bInitialized)
-			{
-				bUpdated = true;
-				ColumnData.bInitialized = true;
-			}
-			ColumnData.Size += CurrentSize;
+			continue;
 		}
-	}
-	else
-	{
-		for (const TPair<const FTagData*, FTrackerTagSizeData>& It : TagSizes)
+		if (SystemsTag->GetDisplayName() == TagName_Untagged)
 		{
-			const FTagData* SystemsTag = It.Key;
-			if (SystemsTag->GetTagSet() != ELLMTagSet::None || !SystemsTag->IsReportable())
-			{
-				continue;
-			}
-			if (SystemsTag->GetDisplayName() == TagName_Untagged)
-			{
-				// We write the "Untagged" tag in the 2nd column, and always write it, rather than writing it in the
-				// column index where we first notice it. That writing is handled by use of OverrideUntaggedTagData in
-				// WriteHeader and AddRow.
-				continue;
-			}
+			// We write the "Untagged" tag in the 2nd column, and always write it, rather than writing it in the
+			// column index where we first notice it. That writing is handled by use of OverrideUntaggedTagData in
+			// WriteHeader and AddRow.
+			continue;
+		}
 
-			FColumnData& ColumnData = Columns.FindOrAdd(FColumnKey{ SystemsTag, nullptr });
-			if (!ColumnData.bInitialized)
-			{
-				bUpdated = true;
-				ColumnData.bInitialized = true;
-			}
-			// ColumnData.Size is unused when !bRecordingCodeOrContent, because we already have
-			// the accumulated value in TagSizes and do not need to accumulate a sum.
+		CsvWriter::FColumnData& ColumnData = Columns.FindOrAdd(CsvWriter::FColumnKey{ SystemsTag, nullptr });
+		if (!ColumnData.bInitialized)
+		{
+			bUpdated = true;
+			ColumnData.bInitialized = true;
 		}
+		// ColumnData.Size is unused when !bRecordingCodeOrContent, because we already have
+		// the accumulated value in TagSizes and do not need to accumulate a sum.
 	}
 	return bUpdated;
 }
@@ -7155,7 +7162,7 @@ void FLLMCsvWriter::WriteHeader(FLLMGlobals& LLMRef, const FTagData* OverrideTra
 	Archive->Seek(0);
 
 	TStringBuilder<256> NameBuffer;
-	auto WriteColumnName = [this, &NameBuffer](const FColumnKey& ColumnKey)
+	auto WriteColumnName = [this, &NameBuffer](const CsvWriter::FColumnKey& ColumnKey)
 	{
 		if (!ColumnKey.SystemsTag)
 		{
@@ -7163,11 +7170,13 @@ void FLLMCsvWriter::WriteHeader(FLLMGlobals& LLMRef, const FTagData* OverrideTra
 		}
 
 		NameBuffer.Reset();
-		ColumnKey.SystemsTag->AppendDisplayPath(NameBuffer);
 		if (ColumnKey.CodeOrContentTag)
 		{
-			NameBuffer << TEXT("/") << ColumnKey.CodeOrContentTag->GetDisplayName();
+			NameBuffer << ColumnKey.CodeOrContentTag->GetDisplayName() << TEXT("/");
 		}
+
+		ColumnKey.SystemsTag->AppendDisplayPath(NameBuffer);
+		
 		NameBuffer << TEXT(",");
 		Write(NameBuffer);
 	};
@@ -7176,9 +7185,9 @@ void FLLMCsvWriter::WriteHeader(FLLMGlobals& LLMRef, const FTagData* OverrideTra
 	// We combine the Code and Content values for ELLMTags::None tags TrackedTotal and Untagged into a single value
 	// and write a column for that, rather than making separate Code and Content values. It is not useful for budgets
 	// to show the Content-specific value of TrackedTotal and Untagged, so we merge the two columns to save space.
-	WriteColumnName(FColumnKey{ OverrideTrackedTotalTagData, nullptr });
-	WriteColumnName(FColumnKey{ OverrideUntaggedTagData, nullptr });
-	for (const TPair<FColumnKey, FColumnData>& Pair: Columns)
+	WriteColumnName(CsvWriter::FColumnKey{ OverrideTrackedTotalTagData, nullptr });
+	WriteColumnName(CsvWriter::FColumnKey{ OverrideUntaggedTagData, nullptr });
+	for (const TPair<CsvWriter::FColumnKey, CsvWriter::FColumnData>& Pair: Columns)
 	{
 		WriteColumnName(Pair.Key);
 	}
@@ -7241,17 +7250,15 @@ void FLLMCsvWriter::AddRow(FLLMGlobals& LLMRef, const FTrackerTagSizeMap& TagSiz
 
 	if (bRecordingCodeOrContent)
 	{
-		for (TPair<FColumnKey, FColumnData>& Pair : Columns)
+		for (TPair<CsvWriter::FColumnKey, CsvWriter::FColumnData>& Pair : Columns)
 		{
 			WriteValue(Pair.Value.Size);
 		}
 	}
-	else
+
+	for (TPair<CsvWriter::FColumnKey, CsvWriter::FColumnData>& Pair : Columns)
 	{
-		for (TPair<FColumnKey, FColumnData>& Pair : Columns)
-		{
-			WriteTag(Pair.Key.SystemsTag);
-		}
+		WriteTag(Pair.Key.SystemsTag);
 	}
 	Write(TEXTVIEW("\n"));
 
@@ -7473,7 +7480,7 @@ FLLMCsvProfilerWriter::FLLMCsvProfilerWriter()
 void FLLMCsvProfilerWriter::Clear()
 {
 #if LLM_CSV_PROFILER_WRITER_ENABLED
-	TagDataToCsvStatName.Empty();
+	ColumnKeyDataToCsvStatName.Empty();
 #endif
 }
 
@@ -7495,17 +7502,31 @@ void FLLMCsvProfilerWriter::Publish(FLLMGlobals& LLMRef, const FTrackerTagSizeMa
 
 	if (OverrideTrackedTotalTagData)
 	{
-		RecordTagToCsv(CsvCategoryIndex, OverrideTrackedTotalTagData, TrackedTotal);
+		RecordTagToCsv(CsvCategoryIndex, CsvWriter::FColumnKey{OverrideTrackedTotalTagData, nullptr}, TrackedTotal);
 	}
 
 	if (OverrideUntaggedTagData)
 	{
 		const FTagData* TagData = LLMRef.FindTagData(TagName_Untagged, ELLMTagSet::None);
 		const FTrackerTagSizeData* AllocationData = TagData ? TagSizes.Find(TagData) : nullptr;
-		RecordTagToCsv(CsvCategoryIndex, OverrideUntaggedTagData,
+		RecordTagToCsv(CsvCategoryIndex, CsvWriter::FColumnKey{OverrideUntaggedTagData, nullptr},
 			AllocationData ? AllocationData->GetSize(SizeParams) : 0);
 	}
 
+	// Initialize some variables that require ActiveSets to have been parsed.
+	bool bRecordingCodeOrContent = Tracker == ELLMTracker::Default 
+		&& LLMRef.IsTagSetRecordingActiveInner(ELLMTagSet::CodeOrContent);
+
+	if (bRecordingCodeOrContent)
+	{
+		CalculateCodeOrContentColumns(Columns, Tracker, LLMRef, LLMRef.ActiveTagsToAllocationGroup, LLMRef.AllocationGroupsMutex);
+
+		for (const TPair<CsvWriter::FColumnKey, CsvWriter::FColumnData>& KVPair : Columns)
+		{
+			RecordTagToCsv(CSV_CATEGORY_INDEX(LLMCodeOrContent), KVPair.Key, KVPair.Value.Size);
+		}
+	}
+
 	for (const TPair<const FTagData*, FTrackerTagSizeData>& It : TagSizes)
 	{
 		const FTagData* TagData = It.Key;
@@ -7519,25 +7540,39 @@ void FLLMCsvProfilerWriter::Publish(FLLMGlobals& LLMRef, const FTrackerTagSizeMa
 			continue;
 		}
 		const FTrackerTagSizeData* AllocationData = TagSizes.Find(TagData);
-		RecordTagToCsv(CsvCategoryIndex, TagData, AllocationData ? AllocationData->GetSize(SizeParams) : 0);
+		RecordTagToCsv(CsvCategoryIndex, CsvWriter::FColumnKey{TagData, nullptr}, AllocationData ? AllocationData->GetSize(SizeParams) : 0);
 	}
 #endif
 }
 
-void FLLMCsvProfilerWriter::RecordTagToCsv(int32 CsvCategoryIndex, const FTagData* TagData, int64 TagSize)
+void FLLMCsvProfilerWriter::RecordTagToCsv(int32 CsvCategoryIndex, const CsvWriter::FColumnKey& ColumnKey, int64 Size)
 {
 #if LLM_CSV_PROFILER_WRITER_ENABLED
+	ensureMsgf(ColumnKey.SystemsTag != nullptr, TEXT("Systems Tag should not be nullptr, check usage."));
+
+	if (ColumnKey.SystemsTag == nullptr)
+	{
+		return;
+	}
+
 	FName NewCsvStatName;
-	FName* CsvStatNamePtr = TagDataToCsvStatName.Find(TagData);
+	FName* CsvStatNamePtr = ColumnKeyDataToCsvStatName.Find(ColumnKey);
 	if (CsvStatNamePtr == nullptr)
 	{
 		TStringBuilder<FName::StringBufferSize> DisplayPath;
-		TagData->GetDisplayPath(DisplayPath);
+
+		// If there is a CodeOrContent tag, prepend to the name
+		if (ColumnKey.CodeOrContentTag != nullptr)
+		{
+			DisplayPath << ColumnKey.CodeOrContentTag->GetDisplayName() << TEXT("/");
+		}
+
+		ColumnKey.SystemsTag->AppendDisplayPath(DisplayPath);
 		NewCsvStatName = FName(DisplayPath);
-		TagDataToCsvStatName.Add(TagData, NewCsvStatName);
+		ColumnKeyDataToCsvStatName.Add(ColumnKey, NewCsvStatName);
 		CsvStatNamePtr = &NewCsvStatName;
 	}
-	FCsvProfiler::RecordCustomStat(*CsvStatNamePtr, CsvCategoryIndex, (float)((double)TagSize / (1024.0 * 1024.0)),
+	FCsvProfiler::RecordCustomStat(*CsvStatNamePtr, CsvCategoryIndex, (float)((double)Size / (1024.0 * 1024.0)),
 		ECsvCustomStatOp::Set);
 #endif
 }
diff --git a/Engine/Source/Runtime/Core/Private/HAL/LLM/LLMPrivate.h b/Engine/Source/Runtime/Core/Private/HAL/LLM/LLMPrivate.h
--- a/Engine/Source/Runtime/Core/Private/HAL/LLM/LLMPrivate.h
+++ b/Engine/Source/Runtime/Core/Private/HAL/LLM/LLMPrivate.h
@@ -24,6 +24,32 @@ void GlobalRegisterTagDeclaration(FLLMTagDeclaration& TagDeclaration);
 namespace UE::LLMPrivate
 {
 
+class FTagData;
+
+namespace CsvWriter
+{
+
+struct FColumnKey
+{
+	const FTagData* SystemsTag = nullptr;
+	const FTagData* CodeOrContentTag = nullptr;
+	bool operator==(const FColumnKey& Other) const;
+
+	friend uint32 GetTypeHash(const FColumnKey& Key)
+	{
+		constexpr uint32 HashPrime = 101;
+		return ::GetTypeHash(Key.SystemsTag) * HashPrime + ::GetTypeHash(Key.CodeOrContentTag);
+	}
+};
+
+struct FColumnData
+{
+	int64 Size = 0;
+	bool bInitialized = false;
+};
+
+} // namespace UE::LLMPrivate::CsvWriter
+
 struct FTagDataNameKey
 {
 	FName Name;
@@ -269,19 +295,7 @@ private:
 		const FTagData* OverrideTrackedTotalName, const FTagData* OverrideUntaggedName, int64 TrackedTotal,
 		UE::LLM::ESizeParams SizeParams);
 
-	struct FColumnKey
-	{
-		const FTagData* SystemsTag = nullptr;
-		const FTagData* CodeOrContentTag = nullptr;
-		bool operator==(const FColumnKey& Other) const;
-	};
-	friend uint32 GetTypeHash(const FColumnKey& Key);
-	struct FColumnData
-	{
-		int64 Size = 0;
-		bool bInitialized = false;
-	};
-	TLLMMap<FColumnKey, FColumnData> Columns;
+	TLLMMap<CsvWriter::FColumnKey, CsvWriter::FColumnData> Columns;
 	FArchive* Archive;
 	double LastWriteTime;
 	int32 WriteCount;
@@ -323,11 +337,13 @@ public:
 		const FTagData* OverrideTrackedTotalName, const FTagData* OverrideUntaggedName, int64 TrackedTotal,
 		UE::LLM::ESizeParams SizeParams);
 protected:
-	void RecordTagToCsv(int32 CsvCategoryIndex, const FTagData* TagData, int64 TagSize);
+	void RecordTagToCsv(int32 CsvCategoryIndex, const CsvWriter::FColumnKey& ColumnKey, int64 Size);
 #if LLM_CSV_PROFILER_WRITER_ENABLED
 	ELLMTracker Tracker;
-	TLLMMap<const FTagData*, FName> TagDataToCsvStatName;
+	TLLMMap<CsvWriter::FColumnKey, FName> ColumnKeyDataToCsvStatName;
 #endif
+private:
+	TLLMMap<CsvWriter::FColumnKey, CsvWriter::FColumnData> Columns;
 };
 
 /** Per-thread state in an LLMTracker. */
