Added a stub to enable mimalloc on Android

#Synced-CL 50395668

diff --git a/Engine/Source/Runtime/Core/Private/Android/AndroidPlatformMemory.cpp b/Engine/Source/Runtime/Core/Private/Android/AndroidPlatformMemory.cpp
--- a/Engine/Source/Runtime/Core/Private/Android/AndroidPlatformMemory.cpp
+++ b/Engine/Source/Runtime/Core/Private/Android/AndroidPlatformMemory.cpp
@@ -6,6 +6,7 @@
 #include "HAL/MallocBinned.h"
 #include "HAL/MallocBinned2.h"
 #include "HAL/MallocBinned3.h"
+#include "HAL/MallocMimalloc.h"
 #include "HAL/MallocAnsi.h"
 #include "Misc/ScopeLock.h"
 #include "unistd.h"
@@ -464,6 +465,14 @@ FMalloc* FAndroidPlatformMemory::BaseAllocator()
 	return (Instance = new FMallocAnsi());
 #endif
 
+#if	!ENABLE_PGO_PROFILE
+	const bool bHasEnoughMemoryForMimalloc = false;	//stub
+	if (bHasEnoughMemoryForMimalloc)
+	{
+		return (Instance = new FMallocMimalloc());
+	}
+#endif
+
 	const bool bHeapProfilingSupported = AndroidHeapProfiling::Init();
 
 #if USE_MALLOC_BINNED3
diff --git a/Engine/Source/Runtime/Core/Public/Android/AndroidPlatform.h b/Engine/Source/Runtime/Core/Public/Android/AndroidPlatform.h
--- a/Engine/Source/Runtime/Core/Public/Android/AndroidPlatform.h
+++ b/Engine/Source/Runtime/Core/Public/Android/AndroidPlatform.h
@@ -64,6 +64,7 @@ typedef FAndroidTypes FPlatformTypes;
 #define PLATFORM_USE_MINIMAL_HANG_DETECTION			1
 #define PLATFORM_CODE_SECTION(Name)					__attribute__((section(Name)))
 #define PLATFORM_ALLOW_ALLOCATIONS_IN_FASYNCWRITER_SERIALIZEBUFFERTOARCHIVE	0
+#define PLATFORM_SUPPORTS_MIMALLOC					1
 
 #define PLATFORM_RETURN_ADDRESS_FOR_CALLSTACKTRACING	PLATFORM_RETURN_ADDRESS
 
