Introduced PLATFORM_STORES_DATA_IN_POINTER_HIGH_BITS to avoid asserts failing on Android with MTE enabled

#rnx

#Synced-CL 50391842

diff --git a/Engine/Source/Runtime/Core/Private/HAL/LLM/LLM.cpp b/Engine/Source/Runtime/Core/Private/HAL/LLM/LLM.cpp
--- a/Engine/Source/Runtime/Core/Private/HAL/LLM/LLM.cpp
+++ b/Engine/Source/Runtime/Core/Private/HAL/LLM/LLM.cpp
@@ -5080,7 +5080,7 @@ void FLLMTracker::PopTag(ELLMTagSet TagSet)
 	State->PopTag(TagSet);
 }
 
-void FLLMTracker::TrackAllocation(const void* Ptr, int64 Size, ELLMTag DefaultTag, ELLMAllocType AllocType,
+void FLLMTracker::TrackAllocation(FInPtr Ptr, int64 Size, ELLMTag DefaultTag, ELLMAllocType AllocType,
 	bool bTrackInMemPro)
 {
 	FLLMThreadState* State = GetOrCreateState();
@@ -5099,7 +5099,7 @@ void FLLMTracker::TrackAllocation(const void* Ptr, int64 Size, ELLMTag DefaultTa
 	}
 }
 
-void FLLMTracker::TrackAllocation(const void* Ptr, int64 Size, FName DefaultTag, ELLMAllocType AllocType,
+void FLLMTracker::TrackAllocation(FInPtr Ptr, int64 Size, FName DefaultTag, ELLMAllocType AllocType,
 	bool bTrackInMemPro)
 {
 	FLLMThreadState* State = GetOrCreateState();
@@ -5118,7 +5118,7 @@ void FLLMTracker::TrackAllocation(const void* Ptr, int64 Size, FName DefaultTag,
 	}
 }
 
-void FLLMTracker::TrackAllocationInActiveTags(const void* Ptr, int64 Size, ELLMAllocType AllocType,
+void FLLMTracker::TrackAllocationInActiveTags(FInPtr Ptr, int64 Size, ELLMAllocType AllocType,
 	FLLMThreadState* State, bool bTrackInMemPro)
 {
 	if (IsPaused(AllocType))
@@ -5152,7 +5152,7 @@ void FLLMTracker::TrackAllocationInActiveTags(const void* Ptr, int64 Size, ELLMA
 #endif
 }
 
-void FLLMTracker::TrackFree(const void* Ptr, ELLMAllocType AllocType, bool bTrackInMemPro)
+void FLLMTracker::TrackFree(FInPtr Ptr, ELLMAllocType AllocType, bool bTrackInMemPro)
 {
 #if !UE_ONLY_USE_PLATFORM_TRACKER
 	// look up the pointer in the tracking map
@@ -5190,7 +5190,7 @@ void FLLMTracker::TrackFree(const void* Ptr, ELLMAllocType AllocType, bool bTrac
 	State->TrackFree(Ptr, Size, Tracker, AllocType, AllocationGroup, bTrackInMemPro);
 }
 
-void FLLMTracker::OnAllocMoved(const void* Dest, const void* Source, ELLMAllocType AllocType)
+void FLLMTracker::OnAllocMoved(FInPtr Dest, FInPtr Source, ELLMAllocType AllocType)
 {
 #if !UE_ONLY_USE_PLATFORM_TRACKER
 	FLLMAllocMap::Values Values;
@@ -6407,7 +6407,7 @@ void FLLMTracker::GetTagsNamesWithAmountFiltered(TMap<FName, uint64>& OutTagsNam
 #endif // !UE_ONLY_USE_PLATFORM_TRACKER
 }
 
-bool FLLMTracker::FindTagsForPtr(void* InPtr, TArray<const FTagData *, TInlineAllocator<static_cast<int32>(ELLMTagSet::Max)>>& OutTags) const
+bool FLLMTracker::FindTagsForPtr(FInPtr InPtr, TArray<const FTagData *, TInlineAllocator<static_cast<int32>(ELLMTagSet::Max)>>& OutTags) const
 {
 #if !UE_ONLY_USE_PLATFORM_TRACKER
 	uint32 Size;
diff --git a/Engine/Source/Runtime/Core/Private/HAL/LLM/LLMPrivate.h b/Engine/Source/Runtime/Core/Private/HAL/LLM/LLMPrivate.h
--- a/Engine/Source/Runtime/Core/Private/HAL/LLM/LLMPrivate.h
+++ b/Engine/Source/Runtime/Core/Private/HAL/LLM/LLMPrivate.h
@@ -403,6 +403,25 @@ public:
 class FLLMTracker
 {
 public:
+	struct FInPtr
+	{
+		FInPtr(const void* InPtr)
+			: Ptr(InPtr) {
+#if PLATFORM_STORES_DATA_IN_POINTER_HIGH_BITS
+			// The values in the high 16 bits of Ptr are used to label the allocation, and are not needed to identify the pointer.
+			// No two pointers will have the same values for the lower 48 bits but different values for the upper 16. We want to use
+			// the upper 16 bits to store extra data of our own in our table that holds the pointer, so zero those bits out before using
+			// them in our AllocationMap.
+			static_assert(sizeof(void*) == 8);
+			Ptr = reinterpret_cast<const void*>(reinterpret_cast<uintptr_t>(Ptr) & 0x0000'ffff'ffff'ffff);
+#endif
+		}
+		operator const void* () 
+		{ 
+			return Ptr; 
+		}
+		const void* Ptr;
+	};
 
 	FLLMTracker(FLLMGlobals& InLLM);
 	~FLLMTracker();
@@ -413,12 +432,12 @@ public:
 	void PushTag(FName Tag, bool bInIsStatTag, ELLMTagSet TagSet);
 	void PushTag(const FTagData* TagData, ELLMTagSet TagSet);
 	void PopTag(ELLMTagSet TagSet);
-	void TrackAllocation(const void* Ptr, int64 Size, FName DefaultTag, ELLMAllocType AllocType, bool bTrackInMemPro);
-	void TrackAllocation(const void* Ptr, int64 Size, ELLMTag DefaultTag, ELLMAllocType AllocType, bool bTrackInMemPro);
-	void TrackFree(const void* Ptr, ELLMAllocType AllocType, bool bTrackInMemPro);
+	void TrackAllocation(FInPtr Ptr, int64 Size, FName DefaultTag, ELLMAllocType AllocType, bool bTrackInMemPro);
+	void TrackAllocation(FInPtr Ptr, int64 Size, ELLMTag DefaultTag, ELLMAllocType AllocType, bool bTrackInMemPro);
+	void TrackFree(FInPtr Ptr, ELLMAllocType AllocType, bool bTrackInMemPro);
 	void TrackMemoryOfActiveTag(int64 Amount, FName DefaultTag, ELLMAllocType AllocType);
 	void TrackMemoryOfActiveTag(int64 Amount, ELLMTag DefaultTag, ELLMAllocType AllocType);
-	void OnAllocMoved(const void* Dest, const void* Source, ELLMAllocType AllocType);
+	void OnAllocMoved(FInPtr Dest, FInPtr Source, ELLMAllocType AllocType);
 
 	void TrackMemory(ELLMTag EnumTag, int64 Amount, ELLMAllocType AllocType);
 	void TrackMemory(FName TagName, ELLMTagSet TagSet, int64 Amount, ELLMAllocType AllocType);
@@ -493,7 +512,7 @@ public:
 	void GetTagsNamesWithAmount(TMap<FName, uint64>& OutTagsNamesWithAmount, ELLMTagSet TagSet = ELLMTagSet::None);
 	void GetTagsNamesWithAmountFiltered(TMap<FName, uint64>& OutTagsNamesWithAmount, ELLMTagSet TagSet, TArray<FLLMTagSetAllocationFilter>& Filters);
 
-	bool FindTagsForPtr(void* InPtr, TArray<const FTagData*, TInlineAllocator<static_cast<int32>(ELLMTagSet::Max)>>& OutTags) const;
+	bool FindTagsForPtr(FInPtr InPtr, TArray<const FTagData*, TInlineAllocator<static_cast<int32>(ELLMTagSet::Max)>>& OutTags) const;
 
 	int64 GetAllocTypeAmount(ELLMAllocType AllocType);
 
@@ -515,7 +534,7 @@ public:
 protected:
 	FLLMThreadState* GetOrCreateState();
 	FLLMThreadState* GetState();
-	void TrackAllocationInActiveTags(const void* Ptr, int64 Size, ELLMAllocType AllocType, FLLMThreadState* State, bool bTrackInMemPro);
+	void TrackAllocationInActiveTags(FInPtr Ptr, int64 Size, ELLMAllocType AllocType, FLLMThreadState* State, bool bTrackInMemPro);
 	void TrackMemoryOfActiveTag(int64 Amount, const FTagData* TagData, ELLMAllocType AllocType, FLLMThreadState* State);
 	void FetchAndClearTagSizes_UpdateGroupAllocs(FAllocationGroup& AllocationGroup,
 		const FAllocationGroupTrackingData& TrackingData, EPruningLevel PruningLevel);
diff --git a/Engine/Source/Runtime/Core/Public/Android/AndroidPlatform.h b/Engine/Source/Runtime/Core/Public/Android/AndroidPlatform.h
--- a/Engine/Source/Runtime/Core/Public/Android/AndroidPlatform.h
+++ b/Engine/Source/Runtime/Core/Public/Android/AndroidPlatform.h
@@ -65,6 +65,7 @@ typedef FAndroidTypes FPlatformTypes;
 #define PLATFORM_CODE_SECTION(Name)					__attribute__((section(Name)))
 #define PLATFORM_ALLOW_ALLOCATIONS_IN_FASYNCWRITER_SERIALIZEBUFFERTOARCHIVE	0
 #define PLATFORM_SUPPORTS_MIMALLOC					1
+#define PLATFORM_STORES_DATA_IN_POINTER_HIGH_BITS	1
 
 #define PLATFORM_RETURN_ADDRESS_FOR_CALLSTACKTRACING	PLATFORM_RETURN_ADDRESS
 
diff --git a/Engine/Source/Runtime/Core/Public/HAL/Platform.h b/Engine/Source/Runtime/Core/Public/HAL/Platform.h
--- a/Engine/Source/Runtime/Core/Public/HAL/Platform.h
+++ b/Engine/Source/Runtime/Core/Public/HAL/Platform.h
@@ -565,6 +565,13 @@
 	#define	PLATFORM_ALLOW_ALLOCATIONS_IN_FASYNCWRITER_SERIALIZEBUFFERTOARCHIVE 1
 #endif
 
+// Our  LowLevelMemTracker wants to store data about an allocation in the high bits, so it assumes high bits are 0.
+// Some platforms store data about an allocation of their own in the high bits. Let LLM know that on those platforms 
+// it is okay to discard those bits when storing a copy of the pointer for reference.
+#if !defined(PLATFORM_STORES_DATA_IN_POINTER_HIGH_BITS)
+	#define PLATFORM_STORES_DATA_IN_POINTER_HIGH_BITS 0
+#endif
+
 #ifndef PLATFORM_HAS_FPlatformVirtualMemoryBlock
 	#define	PLATFORM_HAS_FPlatformVirtualMemoryBlock 1
 #endif
@@ -1313,3 +1320,4 @@ namespace UE::Core::Private
 	)
 
 // IWYU pragma: end_exports
+
