[Android] fdsan support for file handles

This is enabled by default and can be disabled via `UseFdSan=false` configrule.

#rnx

diff --git a/Engine/Source/Runtime/Core/Private/Android/AndroidPlatformFile.cpp b/Engine/Source/Runtime/Core/Private/Android/AndroidPlatformFile.cpp
--- a/Engine/Source/Runtime/Core/Private/Android/AndroidPlatformFile.cpp
+++ b/Engine/Source/Runtime/Core/Private/Android/AndroidPlatformFile.cpp
@@ -10,12 +10,14 @@
 #include "Misc/App.h"
 #include "Misc/Paths.h"
 #include "Android/AndroidJniGameActivity.h"
+#include "HAL/IConsoleManager.h"
 
 #include <dirent.h>
 #include <unistd.h>
 #include <android/asset_manager.h>
 #include <android/asset_manager_jni.h>
 #include <android/storage_manager.h>
+#include <android/fdsan.h>
 #include "Android/AndroidJava.h"
 #include "Containers/Map.h"
 #include <limits>
@@ -57,6 +59,45 @@ DEFINE_LOG_CATEGORY_STATIC(LogAndroidFile, Log, All);
 // make an FTimeSpan object that represents the "epoch" for time_t (from a stat struct)
 const FDateTime AndroidEpoch(1970, 1, 1);
 
+//
+// FDSan has been added to Android in API level 29. However, symbols are weak, so it is possible 
+// to determine the availability in runtime. NDK headers have them under API level define, so
+// we duplicate those here if needed.
+//
+#if __ANDROID_API__ < 29
+extern "C"
+{
+	void android_fdsan_exchange_owner_tag(int fd, uint64_t expected_tag, uint64_t new_tag) __attribute__((__weak__));
+	int android_fdsan_close_with_tag(int fd, uint64_t tag) __attribute__((__weak__));
+}
+#endif // __ANDROID_API__ < 29
+
+namespace UE::Android::Private
+{
+	struct FdSanConfig
+	{
+		FdSanConfig()
+		{
+			if (const FString* ConfigRulesOverride = FPlatformMisc::GetConfigRulesVariable(TEXT("UseFdSan")))
+			{
+				bIsEnabled = ConfigRulesOverride->Equals("true", ESearchCase::IgnoreCase);
+			}
+			else
+			{
+				bIsEnabled = true;
+			}
+		}
+
+		bool bIsEnabled;
+	};
+
+	static bool ShouldUseFdSan()
+	{
+		static FdSanConfig Config;
+		return Config.bIsEnabled;
+	}
+} // UE::Android::Private
+
 namespace
 {
 	FFileStatData AndroidStatToUEFileData(struct stat& FileInfo)
@@ -78,6 +119,103 @@ namespace
 			!(FileInfo.st_mode & S_IWUSR)
 		);
 	}
+
+	class FAndroidUniqueFd final
+	{
+	public:
+		FAndroidUniqueFd() = default;
+
+		explicit FAndroidUniqueFd(int32 InFd)
+		{
+			Reset(InFd);
+		}
+
+		FAndroidUniqueFd(const FAndroidUniqueFd&) = delete;
+		FAndroidUniqueFd& operator=(const FAndroidUniqueFd&) = delete;
+
+		FAndroidUniqueFd(FAndroidUniqueFd&& Other) noexcept
+		{
+			*this = MoveTemp(Other);
+		}
+
+		FAndroidUniqueFd& operator=(FAndroidUniqueFd&& Other) noexcept
+		{
+			if (this == &Other)
+			{
+				return *this;
+			}
+
+			Reset();
+
+			if (Other.Fd != -1)
+			{
+				Fd       = Other.Fd;
+				Other.Fd = -1;
+
+				ExchangeTag(Fd, Other.GetTag(), GetTag());
+			}
+
+			return *this;
+		}
+
+		~FAndroidUniqueFd()
+		{
+			Reset();
+		}
+
+		FORCEINLINE int32 Get() const
+		{
+			return Fd;
+		}
+
+		FORCEINLINE bool IsValid() const
+		{
+			return Fd != -1;
+		}
+
+	private:
+		int32 Fd = -1;
+
+		void Reset(int32 NewFd = -1)
+		{
+			uint64 Tag = GetTag();
+
+			if (Fd != -1)
+			{
+				Close(Fd, Tag);
+				Fd = -1;
+			}
+
+			if (NewFd != -1)
+			{
+				Fd = NewFd;
+				ExchangeTag(Fd, 0, Tag);
+			}
+		}
+
+		uint64 GetTag() const
+		{
+			static_assert(sizeof(uint64) == sizeof(void*));
+			return reinterpret_cast<uint64>(this);
+		}
+
+		static int Close(int32 Fd, uint64 Tag)
+		{
+			if (android_fdsan_close_with_tag && UE::Android::Private::ShouldUseFdSan())
+			{
+				return android_fdsan_close_with_tag(Fd, Tag);
+			}
+			return close(Fd);
+		}
+
+		static void ExchangeTag(int32 Fd, uint64 OldTag, uint64 NewTag)
+		{
+			if (android_fdsan_exchange_owner_tag && UE::Android::Private::ShouldUseFdSan())
+			{
+				android_fdsan_exchange_owner_tag(Fd, OldTag, NewTag);
+			}
+		}
+	};
 }
 
 #define USE_UTIME 0
@@ -200,7 +338,7 @@ AAssetDir* OpenExistingAssetManagerDirectory(AAssetManager* AssetManager, const
 	Android file handle implementation for partial (i.e. parcels) files.
 	This can handle all the variations of accessing data for assets and files.
 **/
-class CORE_API FFileHandleAndroid : public IFileHandle
+class CORE_API FFileHandleAndroid final : public IFileHandle
 {
 	enum { READWRITE_SIZE = 1024 * 1024 };
 
@@ -209,7 +347,7 @@ class CORE_API FFileHandleAndroid : public IFileHandle
 #if LOG_ANDROID_FILE
 		FPlatformMisc::LowLevelOutputDebugStringf(
 			TEXT("FFileHandleAndroid => Asset = %p, Handle = %d, Bounds = [%d,%d)"),
-			File->Asset, File->Handle, int32(Start), int32(Start+Length));
+			File->Asset, File->Handle.Get(), int32(Start), int32(Start+Length));
 #endif
 	}
 
@@ -219,29 +357,24 @@ public:
 	{
 		FString Path;
 		AAsset * Asset;
-		int32 Handle;
+		FAndroidUniqueFd Handle;
 
-		FileReference()
-			: Asset(nullptr), Handle(-1)
+		FileReference() : Asset(nullptr)
 		{
 		}
 
 		FileReference(const FString & path, AAsset * asset)
-			: Path(path), Asset(asset), Handle(0)
+			: Path(path), Asset(asset)
 		{
 		}
 
-		FileReference(const FString & path, int32 handle)
-			: Path(path), Asset(nullptr), Handle(handle)
+		FileReference(const FString & path, FAndroidUniqueFd handle)
+			: Path(path), Asset(nullptr), Handle(MoveTemp(handle))
 		{
 		}
 
 		~FileReference()
 		{
-			if (Handle != -1)
-			{
-				close(Handle);
-			}
 			if (nullptr != Asset)
 			{
 				AAsset_close(Asset);
@@ -256,7 +389,7 @@ public:
 
 	FORCEINLINE void CheckValid() const
 	{
-		check(File.IsValid() && File->Handle != -1);
+		check(File.IsValid() && File->Handle.IsValid());
 	}
 
 	// Invalid handle.
@@ -278,27 +411,27 @@ public:
 	}
 
 	// Handle that covers a subsegment of provided file.
-	FFileHandleAndroid(const FString & path, int32 filehandle, int64 filestart, int64 filelength)
-		: File(MakeShareable(new FileReference(path, filehandle)))
+	FFileHandleAndroid(const FString & path, FAndroidUniqueFd filehandle, int64 filestart, int64 filelength)
+		: File(MakeShareable(new FileReference(path, MoveTemp(filehandle))))
 		, Start(filestart), Length(filelength), CurrentOffset(0)
 	{
 		CheckValid();
 #if UE_ANDROID_FILE_64
-		lseek64(File->Handle, filestart, SEEK_SET);
+		lseek64(File->Handle.Get(), filestart, SEEK_SET);
 #else
-		lseek(File->Handle, filestart, SEEK_SET);
+		lseek(File->Handle.Get(), filestart, SEEK_SET);
 #endif
 		LogInfo();
 	}
 
 	// Handle that covers the entire file content.
-	FFileHandleAndroid(const FString & path, int32 filehandle)
-		: File(MakeShareable(new FileReference(path, filehandle)))
+	FFileHandleAndroid(const FString & path, FAndroidUniqueFd filehandle)
+		: File(MakeShareable(new FileReference(path, MoveTemp(filehandle))))
 		, Start(0), Length(0), CurrentOffset(0)
 	{
 		CheckValid();
-		Length = __lseek(File->Handle, 0, SEEK_END);
-		__lseek(File->Handle, 0, SEEK_SET);
+		Length = __lseek(File->Handle.Get(), 0, SEEK_END);
+		__lseek(File->Handle.Get(), 0, SEEK_SET);
 		LogInfo();
 	}
 
@@ -309,7 +442,7 @@ public:
 	{
 		off64_t OutStart = Start;
 		off64_t OutLength = Length;
-		File->Handle = AAsset_openFileDescriptor64(File->Asset, &OutStart, &OutLength);
+		File->Handle = FAndroidUniqueFd { AAsset_openFileDescriptor64(File->Asset, &OutStart, &OutLength) };
 		Start = OutStart;
 		Length = OutLength;
 		CurrentOffset = Start;
@@ -356,7 +489,7 @@ public:
 #if LOG_ANDROID_FILE
 		FPlatformMisc::LowLevelOutputDebugStringf(
 			TEXT("(%d/%d) FFileHandleAndroid:Read => Path = %s, BytesToRead = %d"),
-			FAndroidTLS::GetCurrentThreadId(), File->Handle,
+			FAndroidTLS::GetCurrentThreadId(), File->Handle.Get(),
 			*(File->Path), int32(BytesToRead));
 #endif
 		check(BytesToRead >= 0);
@@ -372,11 +505,11 @@ public:
 		{
 			int64 ThisSize = FMath::Min<int64>(READWRITE_SIZE, BytesToRead);
 			
-			ThisSize = __pread(File->Handle, Destination, ThisSize, CurrentOffset);
+			ThisSize = __pread(File->Handle.Get(), Destination, ThisSize, CurrentOffset);
 #if LOG_ANDROID_FILE
 			FPlatformMisc::LowLevelOutputDebugStringf(
 				TEXT("(%d/%d) FFileHandleAndroid:Read => Path = %s, ThisSize = %d, destination = %X"),
-				FAndroidTLS::GetCurrentThreadId(), File->Handle,
+				FAndroidTLS::GetCurrentThreadId(), File->Handle.Get(),
 				*(File->Path), int32(ThisSize), Destination);
 #endif
 			if (ThisSize < 0)
@@ -424,7 +557,7 @@ public:
 			int64 ThisSize = FMath::Min<int64>(READWRITE_SIZE, BytesToWrite);
 			check(Source);
 			errno = EINTR;
-			int64 Result = __pwrite(File->Handle, Source, ThisSize, CurrentOffset);
+			int64 Result = __pwrite(File->Handle.Get(), Source, ThisSize, CurrentOffset);
 			if (Result <= 0)
 			{
 				if (errno == EINTR)
@@ -436,7 +569,7 @@ public:
 				int32 SaveErrno = errno;
 				FPlatformMisc::LowLevelOutputDebugStringf(
 					TEXT("(%d/%d) FFileHandleAndroid:Write => Path = %s, this size = %d, CurrentOffset = %d, Source = %p, Result = %d, errno = %d"),
-					FAndroidTLS::GetCurrentThreadId(), File->Handle,
+					FAndroidTLS::GetCurrentThreadId(), File->Handle.Get(),
 					*(File->Path), int32(ThisSize), CurrentOffset, Source, int32(Result), SaveErrno);
 #endif
 				bSuccess = false;
@@ -445,7 +578,7 @@ public:
 #if LOG_ANDROID_FILE
 			FPlatformMisc::LowLevelOutputDebugStringf(
 				TEXT("(%d/%d) FFileHandleAndroid:Write => Path = %s, this size = %d, CurrentOffset = %d, Source = %p, Result = %d"),
-				FAndroidTLS::GetCurrentThreadId(), File->Handle,
+				FAndroidTLS::GetCurrentThreadId(), File->Handle.Get(),
 				*(File->Path), int32(ThisSize), CurrentOffset, Source, int32(Result));
 #endif
 			CurrentOffset += Result;
@@ -458,7 +591,7 @@ public:
 #if LOG_ANDROID_FILE
 		FPlatformMisc::LowLevelOutputDebugStringf(
 			TEXT("(%d/%d) FFileHandleAndroid:Write => Path = %s, final size %d"),
-			FAndroidTLS::GetCurrentThreadId(), File->Handle,
+			FAndroidTLS::GetCurrentThreadId(), File->Handle.Get(),
 			*(File->Path), Length);
 #endif
 		return bSuccess;
@@ -474,8 +607,8 @@ public:
 		}
 
 		return bFullFlush
-			? fsync(File->Handle) == 0
-			: fdatasync(File->Handle) == 0;
+			? fsync(File->Handle.Get()) == 0
+			: fdatasync(File->Handle.Get()) == 0;
 	}
 
 	virtual bool Truncate(int64 NewSize) override
@@ -488,12 +621,12 @@ public:
 		}
 
 		int Result = 0;
-		do { Result = __ftruncate(File->Handle, NewSize); } while (Result < 0 && errno == EINTR);
+		do { Result = __ftruncate(File->Handle.Get(), NewSize); } while (Result < 0 && errno == EINTR);
 		if (Result == 0)
 		{
-			const int CurrentPos = __lseek(File->Handle, 0, SEEK_CUR); 
-			Length = __lseek(File->Handle, 0, SEEK_END);
-			__lseek(File->Handle, CurrentPos, SEEK_SET);
+			const int CurrentPos = __lseek(File->Handle.Get(), 0, SEEK_CUR); 
+			Length = __lseek(File->Handle.Get(), 0, SEEK_END);
+			__lseek(File->Handle.Get(), CurrentPos, SEEK_SET);
 		}
 		return Result == 0;
 	}
@@ -510,7 +643,7 @@ private:
 #if LOG_ANDROID_FILE
 		FPlatformMisc::LowLevelOutputDebugStringf(
 			TEXT("(%d/%d) FFileHandleAndroid:Read => Path = %s, BytesToRead = %d"),
-			FAndroidTLS::GetCurrentThreadId(), File->Handle,
+			FAndroidTLS::GetCurrentThreadId(), File->Handle.Get(),
 			*(File->Path), int32(BytesToRead));
 #endif
 		if (BytesToRead < 0 || Offset - Start < 0 || (BytesToRead + Offset - Start) > Size())
@@ -529,11 +662,11 @@ private:
 		{
 			int64 ThisSize = FMath::Min<int64>(READWRITE_SIZE, BytesToRead);
 
-			ThisSize = __pread(File->Handle, Destination, ThisSize, Offset);
+			ThisSize = __pread(File->Handle.Get(), Destination, ThisSize, Offset);
 #if LOG_ANDROID_FILE
 			FPlatformMisc::LowLevelOutputDebugStringf(
 				TEXT("(%d/%d) FFileHandleAndroid:Read => Path = %s, ThisSize = %d, destination = %X"),
-				FAndroidTLS::GetCurrentThreadId(), File->Handle,
+				FAndroidTLS::GetCurrentThreadId(), File->Handle.Get(),
 				*(File->Path), int32(ThisSize), Destination);
 #endif
 			if (ThisSize < 0)
@@ -654,28 +787,51 @@ public:
 
 		ManifestEntries.Empty();
 
-		// int Handle = open( TCHAR_TO_UTF8(ManifestFileName), O_RDWR );
-		int Handle = open( TCHAR_TO_UTF8(*ManifestPath), O_RDONLY );
-
-		if ( Handle == -1 )
+		FString EntireFile;
 		{
+			FAndroidUniqueFd Handle { open(TCHAR_TO_UTF8(*ManifestPath), O_RDONLY) };
+
+			if (!Handle.IsValid())
+			{
 #if LOG_ANDROID_FILE_MANIFEST
-			FPlatformMisc::LowLevelOutputDebugStringf(TEXT("Failed to open file for read'%s'"), *ManifestFileName);
+				FPlatformMisc::LowLevelOutputDebugStringf(TEXT("Failed to open file for read '%s'"), *ManifestFileName);
 #endif
-			return;
-		}
+				return;
+			}
 
-		FString EntireFile;
-		char Buffer[1024];
-		int BytesRead = 1023;
-		while ( BytesRead == 1023 )
-		{
-			BytesRead = read(Handle, Buffer, 1023);
-			Buffer[BytesRead] = '\0';
-			EntireFile.Append(FString(UTF8_TO_TCHAR(Buffer)));
-		}
+			int BytesRead = 0;
+			{
+				ANSICHAR Buffer[1024];
+				do
+				{
+					BytesRead = read(Handle.Get(), Buffer, UE_ARRAY_COUNT(Buffer) - 1);
+					if (BytesRead < 0)
+					{
+						if (errno == EINTR)
+						{
+							// interrupted by signal, no error
+							continue;
+						}
+						break;
+					}
+					Buffer[BytesRead] = '\0';
+					auto Conv = StringCast<TCHAR>(Buffer);
+					EntireFile.Append(Conv.Get(), Conv.Length());
+				}
+				while (BytesRead > 0);
+			}
 
-		close( Handle );
+			if (BytesRead < 0)
+			{
+#if LOG_ANDROID_FILE_MANIFEST
+				const uint32 LastError = FPlatformMisc::GetLastError();
+				TCHAR ErrorBuffer[1024];
+				FPlatformMisc::GetSystemErrorMessage(ErrorBuffer, UE_ARRAY_COUNT(ErrorBuffer), LastError);
+				FPlatformMisc::LowLevelOutputDebugStringf(TEXT("Failed to read the manifest file '%s' (%u: %s)"), *ManifestFileName, LastError, ErrorBuffer);
+#endif
+				return;
+			}
+		}
 
 		TArray<FString> Lines;
 		EntireFile.ParseIntoArrayLines(Lines);
@@ -730,9 +886,9 @@ public:
 		const FString ManifestPath = BasePath + ManifestFileName;
 
 
-		int Handle = open(TCHAR_TO_UTF8(*ManifestPath), O_WRONLY | O_CREAT | O_TRUNC, S_IRUSR | S_IWUSR);
+		FAndroidUniqueFd Handle { open(TCHAR_TO_UTF8(*ManifestPath), O_WRONLY | O_CREAT | O_TRUNC, S_IRUSR | S_IWUSR) };
 
-		if ( Handle == -1 )
+		if (!Handle.IsValid())
 		{
 #if LOG_ANDROID_FILE_MANIFEST
 			FPlatformMisc::LowLevelOutputDebugStringf(TEXT("Failed to open file for write '%s'"), *ManifestFileName);
@@ -740,20 +896,51 @@ public:
 			return;
 		}
 
-
-		for ( const auto& Line : ManifestEntries )
+		for (const auto& Line : ManifestEntries)
 		{
-			const int BufferSize = 4096;
-			char Buffer[BufferSize] = {"\0"}; 
-			const FString RawDateTimeString = Line.Value.ToIso8601();
-			const FString DateTimeString = FString::Printf(TEXT("%s\r\n"), *RawDateTimeString);
-			strncpy(Buffer, TCHAR_TO_UTF8(*Line.Key), BufferSize-1);
-			strncat(Buffer, "\t", BufferSize-1 );
-			strncat(Buffer, TCHAR_TO_UTF8(*DateTimeString), BufferSize-1 );
-			write( Handle, Buffer, strlen( Buffer ) );
-		}
+			bool bSuccess = true;
+			{
+				const FString RawDateTimeString = Line.Value.ToIso8601();
+
+				ANSICHAR Buffer[4096] = {0};
+				FCStringAnsi::Snprintf(
+					Buffer,
+					UE_ARRAY_COUNT(Buffer),
+					"%s\t%s\r\n",
+					StringCast<UTF8CHAR>(*Line.Key).Get(),
+					StringCast<UTF8CHAR>(*RawDateTimeString).Get());
+
+				ANSICHAR* Source = Buffer;
+				int64 BytesToWrite = FCStringAnsi::Strlen(Source);
+				while (BytesToWrite > 0)
+				{
+					int64 Result = write(Handle.Get(), Source, BytesToWrite);
+					if (Result <= 0)
+					{
+						if (errno == EINTR)
+						{
+							// interrupted by signal, no error
+							continue;
+						}
+						bSuccess = false;
+						break;
+					}
 
-		close( Handle );
+					Source       += Result;
+					BytesToWrite -= Result;
+				}
+			}
+			if (!bSuccess)
+			{
+#if LOG_ANDROID_FILE_MANIFEST
+				const uint32 LastError = FPlatformMisc::GetLastError();
+				TCHAR ErrorBuffer[1024];
+				FPlatformMisc::GetSystemErrorMessage(ErrorBuffer, UE_ARRAY_COUNT(ErrorBuffer), LastError);
+				FPlatformMisc::LowLevelOutputDebugStringf(TEXT("Failed to write the manifest file '%s' (%u: %s)"), *ManifestFileName, LastError, ErrorBuffer);
+#endif
+				return;
+			}
+		}
 	}
 };
 
@@ -1140,19 +1327,18 @@ class FAndroidMappedFileHandle final : public IMappedFileHandle
 	inline static SIZE_T FileMappingAlignment = FPlatformMemory::GetConstants().PageSize;
 
 public:
-	FAndroidMappedFileHandle(int InFileHandle, int64 FileSize, const FString& InFilename)
+	FAndroidMappedFileHandle(FAndroidUniqueFd InFileHandle, int64 FileSize, const FString& InFilename)
 		: IMappedFileHandle(FileSize)
 		, MappedPtr(nullptr)
 		, Filename(InFilename)
 		, NumOutstandingRegions(0)
-		, FileHandle(InFileHandle)
+		, FileHandle(MoveTemp(InFileHandle))
 	{
 	}
 
 	virtual ~FAndroidMappedFileHandle() override
 	{
 		check(!NumOutstandingRegions); // can't delete the file before you delete all outstanding regions
-		close(FileHandle);
 	}
 
 	virtual IMappedFileRegion* MapRegion(int64 Offset = 0, int64 BytesToMap = MAX_int64, FFileMappingFlags Flags = EMappedFileFlags::ENone) override
@@ -1178,7 +1364,7 @@ public:
 			InternalFlags |= MAP_PRIVATE;
 		}
 
-		const uint8* AlignedMapPtr = static_cast<const uint8*>(mmap(nullptr, AlignedSize, Protection, InternalFlags, FileHandle, AlignedOffset));
+		const uint8* AlignedMapPtr = static_cast<const uint8*>(mmap(nullptr, AlignedSize, Protection, InternalFlags, FileHandle.Get(), AlignedOffset));
 		if (AlignedMapPtr == MAP_FAILED || AlignedMapPtr == nullptr)
 		{
 #if LOG_ANDROID_FILE
@@ -1212,13 +1398,13 @@ private:
 	const uint8* MappedPtr;
 	FString Filename;
 	int32 NumOutstandingRegions;
-	int FileHandle;
+	FAndroidUniqueFd FileHandle;
 
 	int64 GetCurrentFileSize() const
 	{
 		struct stat FileInfo;
 		FileInfo.st_size = -1;
-		const int StatResult = fstat(FileHandle, &FileInfo);
+		const int StatResult = fstat(FileHandle.Get(), &FileInfo);
 		if (StatResult == -1)
 		{
 			const int ErrNo = errno;
@@ -1300,13 +1486,14 @@ public:
 		{
 			// Open the APK as a ZIP
 			FZipUnionFile APKZip;
-			int32 Handle = open(TCHAR_TO_UTF8(*GAPKFilename), O_RDONLY);
-			if (Handle == -1)
+			FAndroidUniqueFd Handle { open(TCHAR_TO_UTF8(*GAPKFilename), O_RDONLY) };
+
+			if (!Handle.IsValid())
 			{
 				FPlatformMisc::LowLevelOutputDebugStringf(TEXT("FAndroidPlatformFile::Initialize unable to open APK: %s"), *GAPKFilename);
 				return false;
 			}
-			FFileHandleAndroid* APKFile = new FFileHandleAndroid(GAPKFilename, Handle);
+			FFileHandleAndroid* APKFile = new FFileHandleAndroid(GAPKFilename, MoveTemp(Handle));
 			APKZip.AddPatchFile(MakeShareable(APKFile));
 
 			// Now open the OBB in the APK and mount it
@@ -1501,8 +1688,8 @@ public:
 		const FString NormalizedFilename = LocalPath;
 
 		const int Flags = EnumHasAnyFlags(OpenOptions, EOpenReadFlags::AllowWrite) ? O_RDWR : O_RDONLY;
-		const int32 Handle = open(TCHAR_TO_UTF8(*NormalizedFilename), Flags);
-		if (Handle == -1)
+		FAndroidUniqueFd Handle { open(TCHAR_TO_UTF8(*NormalizedFilename), Flags) };
+		if (!Handle.IsValid())
 		{
 			const int ErrNo = errno;
 			FString ErrorStr = FString::Printf(TEXT("FAndroidPlatformFile::OpenMappedEx('%s', Flags=0x%08X) failed: errno=%d (%s)"), *NormalizedFilename, Flags, ErrNo, UTF8_TO_TCHAR(strerror(ErrNo)));
@@ -1514,7 +1701,7 @@ public:
 
 		struct stat FileInfo;
 		FileInfo.st_size = -1;
-		const int StatResult = fstat(Handle, &FileInfo);
+		const int StatResult = fstat(Handle.Get(), &FileInfo);
 		if (StatResult == -1)
 		{
 			const int ErrNo = errno;
@@ -1527,7 +1714,7 @@ public:
 			return MakeError(MoveTemp(ErrorStr));
 		}
 
-		return MakeValue(new FAndroidMappedFileHandle(Handle, FileInfo.st_size, NormalizedFilename));
+		return MakeValue(new FAndroidMappedFileHandle(MoveTemp(Handle), FileInfo.st_size, NormalizedFilename));
 	}
 
 	virtual int64 FileSize(const TCHAR* Filename) override
@@ -1957,10 +2144,10 @@ public:
 
 		if (IsLocal(LocalPath))
 		{
-			int32 Handle = open(TCHAR_TO_UTF8(*LocalPath), O_RDONLY);
-			if (Handle != -1)
+			FAndroidUniqueFd Handle { open(TCHAR_TO_UTF8(*LocalPath), O_RDONLY) };
+			if (Handle.IsValid())
 			{
-				return new FFileHandleAndroid(LocalPath, Handle);
+				return new FFileHandleAndroid(LocalPath, MoveTemp(Handle));
 			}
 		}
 		else if (IsResource(AssetPath))
@@ -2013,10 +2200,10 @@ public:
 			Flags |= O_WRONLY;
 		}
 
-		int32 Handle = open(TCHAR_TO_UTF8(*LocalPath), Flags, S_IRUSR | S_IWUSR | S_IRGRP | S_IWGRP);
-		if (Handle != -1)
+		FAndroidUniqueFd Handle { open(TCHAR_TO_UTF8(*LocalPath), Flags, S_IRUSR | S_IWUSR | S_IRGRP | S_IWGRP) };
+		if (Handle.IsValid())
 		{
-			FFileHandleAndroid* FileHandleAndroid = new FFileHandleAndroid(LocalPath, Handle);
+			FFileHandleAndroid* FileHandleAndroid = new FFileHandleAndroid(LocalPath, MoveTemp(Handle));
 			if (bAppend)
 			{
 				FileHandleAndroid->SeekFromEnd(0);
