[MonitoredProcess] Adding a simple RunProcessSynchronous static method to FMonitoredProcess to simplify launching a process, waiting, and getting its output.

#Synced-CL 50283614

diff --git a/Engine/Source/Runtime/Core/Private/Misc/MonitoredProcess.cpp b/Engine/Source/Runtime/Core/Private/Misc/MonitoredProcess.cpp
--- a/Engine/Source/Runtime/Core/Private/Misc/MonitoredProcess.cpp
+++ b/Engine/Source/Runtime/Core/Private/Misc/MonitoredProcess.cpp
@@ -28,7 +28,40 @@ FMonitoredProcess::FMonitoredProcess( const FString& InURL, const FString& InPar
 { 
 }
 
- 
+bool FMonitoredProcess::RunProcessSynchronous( const FString& InURL, const FString& InParams, TArray<FString>& OutOutput, float Timeout/*0.f*/ )
+{
+	FMonitoredProcess MonitoredProcess(InURL, InParams, /*InHidden=*/true);
+
+	MonitoredProcess.OnOutput().BindLambda([&OutOutput](const FString& Output)
+		{
+			OutOutput.Add(Output);
+		});
+
+	if (!MonitoredProcess.Launch())
+	{
+		return false;
+	}
+
+	double StartTime = FPlatformTime::Seconds();
+	while (MonitoredProcess.Update())
+	{
+		if (Timeout > 0.f)
+		{
+			if ((FPlatformTime::Seconds() - StartTime) >= Timeout)
+			{
+				MonitoredProcess.Cancel(true);
+				return false;
+			}
+		}
+		FPlatformProcess::Sleep(0.01f);
+	}
+
+	// Add any remaining unprocessed output.
+	OutOutput.Add(MonitoredProcess.GetFullOutputWithoutDelegate());
+
+	return MonitoredProcess.GetReturnCode() == 0;
+}
+
 FMonitoredProcess::~FMonitoredProcess()
 {
 	if (bIsRunning)
@@ -43,7 +76,6 @@ FMonitoredProcess::~FMonitoredProcess()
 	}
 }
 
-
 /* FMonitoredProcess interface
  *****************************************************************************/
 
diff --git a/Engine/Source/Runtime/Core/Public/Misc/MonitoredProcess.h b/Engine/Source/Runtime/Core/Public/Misc/MonitoredProcess.h
--- a/Engine/Source/Runtime/Core/Public/Misc/MonitoredProcess.h
+++ b/Engine/Source/Runtime/Core/Public/Misc/MonitoredProcess.h
@@ -58,6 +58,17 @@ public:
 	*/
 	CORE_API FMonitoredProcess( const FString& InURL, const FString& InParams, const FString& InWorkingDir, bool InHidden, bool InCreatePipes = true );
 
+	/**
+	 * Helper method for synchronously executing a process and gathering the output
+	 *
+	 * @param InURL The URL of the executable to launch
+	 * @param InParams The command line parameters
+	 * @param OutOutput An array of strings containing the output
+	 * @param Timeout An optional maximum time to wait for the process to exit before cancelling it.
+	 * @return Returns true if the process completed successfully (return code == 0)
+	 */
+	static CORE_API bool RunProcessSynchronous(const FString& InURL, const FString& InParams, TArray<FString>& OutOutput, float Timeout = 0.f);
+
 	/** Destructor. */
 	CORE_API virtual ~FMonitoredProcess();
 
