Change android's native logging to INFO priority, this reinstates logging via AGDE's log window.

#tests ran on local device.
[RN] changed default android log level from ANDROID_LOG_DEBUG to ANDROID_LOG_INFO. This prevents AGDE from filtering log output.

#Synced-CL 50398297

diff --git a/Engine/Source/Runtime/Core/Private/Android/AndroidPlatformMisc.cpp b/Engine/Source/Runtime/Core/Private/Android/AndroidPlatformMisc.cpp
--- a/Engine/Source/Runtime/Core/Private/Android/AndroidPlatformMisc.cpp
+++ b/Engine/Source/Runtime/Core/Private/Android/AndroidPlatformMisc.cpp
@@ -72,6 +72,18 @@ extern int32 GAndroidPackageVersion;
 #define GAndroidPackageVersion	1
 #endif
 
+#ifndef FORCE_ANDROID_LOG_TO_INFO_LEVEL
+	#define FORCE_ANDROID_LOG_TO_INFO_LEVEL 0
+#endif
+
+// anything using LocalPrint (LowLevelOutputDebugString and therefore more generally, UE_LOG etc) will be sent to logcat via this log level.
+// Note that <= ANDROID_LOG_DEBUG maybe filtered out by VS/AGDE.
+android_LogPriority AndroidDefaultLogPriority()
+{
+	static const bool bPromoteLogLevel = !UE_BUILD_SHIPPING && (FORCE_ANDROID_LOG_TO_INFO_LEVEL || FPlatformMisc::IsDebuggerPresent());
+	return bPromoteLogLevel ? ANDROID_LOG_INFO : ANDROID_LOG_DEBUG;
+}
+
 static int32 GAndroidTraceMarkersEnabled = 0;
 static FAutoConsoleVariableRef CAndroidTraceMarkersEnabled(
 	TEXT("android.tracemarkers"),
@@ -345,7 +357,7 @@ void FAndroidMisc::LocalPrint(const TCHAR *Message)
 			}
 		}
 		*WritePtr = '\0';
-		__android_log_print(ANDROID_LOG_DEBUG, "UE", "%ls", MessageBuffer);
+		__android_log_print(AndroidDefaultLogPriority(), "UE", "%ls", MessageBuffer);
 	}
 #endif
 }
