Add HashingArchiveProxy to MemoryHasher, replacing a private one.

#Synced-CL 50366285

diff --git a/Engine/Source/Developer/DerivedDataCache/Private/HashingArchiveProxy.h b/Engine/Source/Developer/DerivedDataCache/Private/HashingArchiveProxy.h
--- a/Engine/Source/Developer/DerivedDataCache/Private/HashingArchiveProxy.h
+++ b/Engine/Source/Developer/DerivedDataCache/Private/HashingArchiveProxy.h
@@ -2,37 +2,4 @@
 
 #pragma once
 
-#include "Serialization/ArchiveProxy.h"
-#include "Templates/UnrealTemplate.h"
-
-namespace UE
-{
-
-template <typename HashBuilderType>
-class THashingArchiveProxy : public FArchiveProxy
-{
-public:
-	using FArchiveProxy::FArchiveProxy;
-
-	void Serialize(void* V, int64 Length) override
-	{
-		FArchiveProxy::Serialize(V, Length);
-		HashBuilder.Update(V, uint64(Length));
-	}
-
-	void Seek(int64 InPos) override
-	{
-		checkf(InPos == Tell(), TEXT("A hash cannot be computed when serialization relies on seeking."));
-		FArchiveProxy::Seek(InPos);
-	}
-
-	auto GetHash() -> decltype(DeclVal<HashBuilderType>().Finalize())
-	{
-		return HashBuilder.Finalize();
-	}
-
-private:
-	HashBuilderType HashBuilder;
-};
-
-} // UE
+#include "Serialization/MemoryHasher.h"
diff --git a/Engine/Source/Runtime/Core/Public/Serialization/MemoryHasher.h b/Engine/Source/Runtime/Core/Public/Serialization/MemoryHasher.h
--- a/Engine/Source/Runtime/Core/Public/Serialization/MemoryHasher.h
+++ b/Engine/Source/Runtime/Core/Public/Serialization/MemoryHasher.h
@@ -13,6 +13,7 @@
 #include "Serialization/MemoryArchive.h"
 #include "Misc/SecureHash.h"
 #include "Hash/Blake3.h"
+#include "Serialization/ArchiveProxy.h"
 
 /**
  * Archive for hashing arbitrary data
@@ -59,3 +60,38 @@ protected:
 
 using FMemoryHasherSHA1 = TMemoryHasher<FSHA1, FSHAHash>;
 using FMemoryHasherBlake3 = TMemoryHasher<FBlake3, FBlake3Hash>;
+
+/**
+ * Archive for hashing another archive
+ */
+namespace UE
+{
+
+template <typename HashBuilderType>
+class THashingArchiveProxy : public FArchiveProxy
+{
+public:
+	using FArchiveProxy::FArchiveProxy;
+
+	void Serialize(void* V, int64 Length) override
+	{
+		FArchiveProxy::Serialize(V, Length);
+		HashBuilder.Update(V, uint64(Length));
+	}
+
+	void Seek(int64 InPos) override
+	{
+		checkf(InPos == Tell(), TEXT("A hash cannot be computed when serialization relies on seeking."));
+		FArchiveProxy::Seek(InPos);
+	}
+
+	auto GetHash() -> decltype(DeclVal<HashBuilderType>().Finalize())
+	{
+		return HashBuilder.Finalize();
+	}
+
+private:
+	HashBuilderType HashBuilder;
+};
+
+} // UE
