Add support for applying boot hotfix cvars via the commandline using -BootHotfixCVars="r.mycvar=foo,r.myOtherCvar=bar,...".
This is useful for testing cvars which need to be set before engine initialization.

#Synced-CL 50391842

diff --git a/Engine/Source/Runtime/Core/Private/Misc/ConfigUtilities.cpp b/Engine/Source/Runtime/Core/Private/Misc/ConfigUtilities.cpp
--- a/Engine/Source/Runtime/Core/Private/Misc/ConfigUtilities.cpp
+++ b/Engine/Source/Runtime/Core/Private/Misc/ConfigUtilities.cpp
@@ -1,6 +1,8 @@
 // Copyright Epic Games, Inc. All Rights Reserved.
 
 #include "Misc/ConfigUtilities.h"
+
+#include "Misc/CommandLine.h"
 #include "Misc/ConfigCacheIni.h"
 #include "HAL/FileManager.h"
 #include "HAL/IConsoleManager.h"
@@ -273,6 +275,25 @@ void ApplyCVarsFromBootHotfix()
 	TMap<FString, FString> CVarsToApply;
 	LoadCVarsFromFileForNextBoot(CVarsToApply);
 
+#if !UE_BUILD_SHIPPING
+	// Support passing boot hotfix cvars via the commandline for testing purposes
+	FString BootCvars;
+	if (FParse::Value(FCommandLine::Get(), TEXT("BootHotfixCvars="), BootCvars))
+	{
+		TArray<FString> Pairs;
+		BootCvars.ParseIntoArray(Pairs, TEXT(","), true);
+		for (const FString& Pair : Pairs)
+		{
+			FString Key, Value;
+			if (Pair.Split(TEXT("="), &Key, &Value))
+			{
+				CVarsToApply.Add(Key, Value);
+				UE_LOG(LogConfig, Log, TEXT("Applied commandline boot hotfix cvar %s=%s"), *Key, *Value);
+			}
+		}
+	}
+#endif //!UE_BUILD_SHIPPING
+
 	for (auto It = CVarsToApply.CreateIterator(); It; ++It)
 	{
 		IConsoleVariable* CVar = IConsoleManager::Get().FindConsoleVariable(*It->Key); 
