CsvProfiler - Add metadata for whether vsync is enabled or not.

#rnx

#Synced-CL 50391842

#Synced-CL 50391842

diff --git a/Engine/Source/Runtime/Core/Private/ProfilingDebugging/CsvProfiler.cpp b/Engine/Source/Runtime/Core/Private/ProfilingDebugging/CsvProfiler.cpp
--- a/Engine/Source/Runtime/Core/Private/ProfilingDebugging/CsvProfiler.cpp
+++ b/Engine/Source/Runtime/Core/Private/ProfilingDebugging/CsvProfiler.cpp
@@ -3918,6 +3918,12 @@ void FCsvProfiler::BeginCaptureInternal(const FCsvCaptureCommand& CurrentCommand
 	SetMetadataInternal(TEXT("StartTimestamp"), *FString::Printf(TEXT("%lld"), FDateTime::UtcNow().ToUnixTimestamp()));
 	SetMetadataInternal(TEXT("NamedEvents"), (GCycleStatsShouldEmitNamedEvents > 0) ? TEXT("1") : TEXT("0"));
 
+	static IConsoleVariable* VSyncCVar = IConsoleManager::Get().FindConsoleVariable(TEXT("r.VSync"));
+	if (VSyncCVar)
+	{
+		SetMetadataInternal(TEXT("VSyncEnabled"), *FString::FromInt(VSyncCVar->GetInt()));
+	}
+
 	if (FPlatformMemory::GetProgramSize() > 0)
 	{
 		// Some platforms adjust program size at runtime (based on DLL initialization), so with do this on start capture rather than in CsvProfiler::Init()
