Optimized variable-length integer encoding

Based on a suggestion from Harris.Barra. This has been profiled against multiple branchless versions of the calculation. The fastest branchless implementation is still 50% slower than this when profiled on x86-64.

#Synced-CL 50253651

diff --git a/Engine/Source/Runtime/Core/Public/Serialization/VarInt.h b/Engine/Source/Runtime/Core/Public/Serialization/VarInt.h
--- a/Engine/Source/Runtime/Core/Public/Serialization/VarInt.h
+++ b/Engine/Source/Runtime/Core/Public/Serialization/VarInt.h
@@ -72,13 +72,33 @@ FORCEINLINE uint32 MeasureVarInt(const void* InData)
 /** Measure the number of bytes (1-5) required to encode the 32-bit input. */
 FORCEINLINE uint32 MeasureVarUInt(uint32 InValue)
 {
-	return uint32(int32(FPlatformMath::FloorLog2(InValue)) / 7 + 1);
+	// Retaining this version with no branching until both have been profiled on ARM.
+	// This is equivalent to Min(FloorLog2(N) / 7 + 1, 9)
+	//return FPlatformMath::Min(((uint32(FPlatformMath::FloorLog2NonZero(InValue | 1)) * 37) >> 8) + 1, 9u);
+
+	// This branching implementation is faster than branchless alternatives in real-world benchmarks.
+	return (InValue < (uint32(1) << 28))
+		? ((InValue < (uint32(1) << 14))
+			? ((InValue < (uint32(1) <<  7)) ? 1 : 2)
+			: ((InValue < (uint32(1) << 21)) ? 3 : 4))
+		: 5;
 }
 
 /** Measure the number of bytes (1-9) required to encode the 64-bit input. */
 FORCEINLINE uint32 MeasureVarUInt(uint64 InValue)
 {
-	return uint32(FPlatformMath::Min(int32(FPlatformMath::FloorLog2_64(InValue)) / 7 + 1, 9));
+	// Retaining this version with no branching until both have been profiled on ARM.
+	// This is equivalent to Min(FloorLog2(N) / 7 + 1, 9)
+	//return FPlatformMath::Min(((uint32(FPlatformMath::FloorLog2NonZero_64(InValue | 1)) * 37) >> 8) + 1, 9u);
+
+	// This branching implementation is faster than branchless alternatives in real-world benchmarks.
+	return (InValue < (uint64(1) << 28))
+		? ((InValue < (uint64(1) << 14))
+			? ((InValue < (uint64(1) <<  7)) ? 1 : 2)
+			: ((InValue < (uint64(1) << 21)) ? 3 : 4))
+		: ((InValue < (uint64(1) << 42))
+			? ((InValue < (uint64(1) << 35)) ? 5 : 6)
+			: ((InValue < (uint64(1) << 49)) ? 7 : ((InValue < (uint64(1) << 56)) ? 8 : 9)));
 }
 
 /** Measure the number of bytes (1-5) required to encode the 32-bit input. \see \ref MeasureVarUInt */
