[DiskUtilizationTracker] Hook up DiskUtilizationTracker for Android/IOS

#rnx
#tests Android/IOS RRs

#Synced-CL 50152864

diff --git a/Engine/Source/Runtime/Core/Private/Android/AndroidPlatformFile.cpp b/Engine/Source/Runtime/Core/Private/Android/AndroidPlatformFile.cpp
--- a/Engine/Source/Runtime/Core/Private/Android/AndroidPlatformFile.cpp
+++ b/Engine/Source/Runtime/Core/Private/Android/AndroidPlatformFile.cpp
@@ -20,6 +20,8 @@
 #include <android/fdsan.h>
 #include "Android/AndroidJava.h"
 #include "Containers/Map.h"
+#include "ProfilingDebugging/PlatformFileTrace.h"
+#include "HAL/DiskUtilizationTracker.h"
 #include <limits>
 #include <sys/mman.h>
 #include "Misc/ScopeLock.h"
@@ -501,6 +503,14 @@ public:
 			return false;
 		}
 
+		int64 TotalBytesRead = 0;
+		FScopedDiskUtilizationTracker Tracker(BytesToRead, CurrentOffset);
+		TRACE_PLATFORMFILE_BEGIN_READ(this, File->Handle.Get(), CurrentOffset, BytesToRead);
+		ON_SCOPE_EXIT
+		{
+			TRACE_PLATFORMFILE_END_READ(this, TotalBytesRead);
+		};
+
 		do
 		{
 			int64 ThisSize = FMath::Min<int64>(READWRITE_SIZE, BytesToRead);
@@ -527,6 +537,7 @@ public:
 				bAtEndOfFile = true;
 				break;
 			}
+			TotalBytesRead += ThisSize;
 			CurrentOffset += ThisSize;
 			Destination += ThisSize;
 			BytesToRead -= ThisSize;
@@ -658,6 +669,14 @@ private:
 
 		check(Destination);
 
+		int64 TotalBytesRead = 0;
+		FScopedDiskUtilizationTracker Tracker(BytesToRead, Offset);
+		TRACE_PLATFORMFILE_BEGIN_READ(this, File->Handle.Get(), Offset, BytesToRead);
+		ON_SCOPE_EXIT
+		{
+			TRACE_PLATFORMFILE_END_READ(this, TotalBytesRead);
+		};
+
 		while (BytesToRead > 0)
 		{
 			int64 ThisSize = FMath::Min<int64>(READWRITE_SIZE, BytesToRead);
@@ -677,6 +696,7 @@ private:
 			{
 				break;
 			}
+			TotalBytesRead += ThisSize;
 			Offset += ThisSize;
 			Destination += ThisSize;
 			BytesToRead -= ThisSize;
diff --git a/Engine/Source/Runtime/Core/Private/Apple/ApplePlatformFile.cpp b/Engine/Source/Runtime/Core/Private/Apple/ApplePlatformFile.cpp
--- a/Engine/Source/Runtime/Core/Private/Apple/ApplePlatformFile.cpp
+++ b/Engine/Source/Runtime/Core/Private/Apple/ApplePlatformFile.cpp
@@ -23,6 +23,7 @@
 #include "HAL/PlatformTime.h"
 #include "Misc/ScopeLock.h"
 #include "ProfilingDebugging/PlatformFileTrace.h"
+#include "HAL/DiskUtilizationTracker.h"
 #include "Templates/Function.h"
 #include <sys/stat.h>
 
@@ -218,7 +219,12 @@ public:
 		}
 		
 		int64 TotalBytesRead = 0;
+		FScopedDiskUtilizationTracker Tracker(BytesToRead, Offset);
 		TRACE_PLATFORMFILE_BEGIN_READ(this, FileHandle, Offset, BytesToRead);
+		ON_SCOPE_EXIT
+		{
+			TRACE_PLATFORMFILE_END_READ(this, TotalBytesRead);
+		};
 
 		do
 		{
@@ -231,7 +237,6 @@ public:
 			}
 			if (BytesRead <= 0)
 			{
-				TRACE_PLATFORMFILE_END_READ(this, TotalBytesRead);
 				return false;
 			}
 
@@ -242,8 +247,6 @@ public:
 
 		} while (BytesToRead > 0);
 
-		TRACE_PLATFORMFILE_END_READ(this, TotalBytesRead);
-
 		return true;
 	}
 	
@@ -330,8 +333,15 @@ private:
 	{
 		check(IsValid());
 		int64 MaxReadSize = READWRITE_SIZE;
-		int64 BytesRead = 0;
-		TRACE_PLATFORMFILE_BEGIN_READ(this, FileHandle, 0, BytesToRead);
+
+		int64 TotalBytesRead = 0;
+		FScopedDiskUtilizationTracker Tracker(BytesToRead, FileOffset);
+		TRACE_PLATFORMFILE_BEGIN_READ(this, FileHandle, FileOffset, BytesToRead);
+		ON_SCOPE_EXIT
+		{
+			TRACE_PLATFORMFILE_END_READ(this, TotalBytesRead);
+		};
+
 		while (BytesToRead)
 		{
 			check(BytesToRead >= 0);
@@ -346,20 +356,17 @@ private:
 					MaxReadSize /= 2;
 					continue;
 				}
-				TRACE_PLATFORMFILE_END_READ(this, BytesRead);
-				return BytesRead;
+				return TotalBytesRead;
 			}
-			BytesRead += ThisRead;
+			TotalBytesRead += ThisRead;
 			if (ThisRead != ThisSize)
 			{
-				TRACE_PLATFORMFILE_END_READ(this, BytesRead);
-				return BytesRead;
+				return TotalBytesRead;
 			}
 			Destination += ThisSize;
 			BytesToRead -= ThisSize;
 		}
-		TRACE_PLATFORMFILE_END_READ(this, BytesRead);
-		return BytesRead;
+		return TotalBytesRead;
 	}
 
 	// Holds the internal file handle.
diff --git a/Engine/Source/Runtime/Core/Private/IOS/IOSPlatformFile.cpp b/Engine/Source/Runtime/Core/Private/IOS/IOSPlatformFile.cpp
--- a/Engine/Source/Runtime/Core/Private/IOS/IOSPlatformFile.cpp
+++ b/Engine/Source/Runtime/Core/Private/IOS/IOSPlatformFile.cpp
@@ -23,6 +23,8 @@
 #include "Misc/App.h"
 #include "Misc/Paths.h"
 #include "Misc/ScopeLock.h"
+#include "ProfilingDebugging/PlatformFileTrace.h"
+#include "HAL/DiskUtilizationTracker.h"
 
 #include "Async/MappedFileHandle.h"
 
@@ -155,6 +157,13 @@ public:
 	int64 InternalRead(uint8* Destination, int64 BytesToRead)
 	{
 		int64 TotalBytesRead = 0;
+		FScopedDiskUtilizationTracker Tracker(BytesToRead, FileOffset);
+		TRACE_PLATFORMFILE_BEGIN_READ(this, FileHandle, FileOffset, BytesToRead);
+		ON_SCOPE_EXIT
+		{
+			TRACE_PLATFORMFILE_END_READ(this, TotalBytesRead);
+		};
+
 		while (BytesToRead)
 		{
 			check(BytesToRead >= 0);
@@ -230,6 +239,14 @@ public:
 		}
 #endif //MANAGE_FILE_HANDLES_IOS
 
+		int64 TotalBytesRead = 0;
+		FScopedDiskUtilizationTracker Tracker(BytesToRead, Offset);
+		TRACE_PLATFORMFILE_BEGIN_READ(this, FileHandle, Offset, BytesToRead);
+		ON_SCOPE_EXIT
+		{
+			TRACE_PLATFORMFILE_END_READ(this, TotalBytesRead);
+		};
+
 		do
 		{
 			size_t BytesToRead32 = static_cast<size_t>(FMath::Min<int64>(READWRITE_SIZE, BytesToRead));
@@ -244,6 +261,7 @@ public:
 				return false;
 			}
 
+			TotalBytesRead += BytesRead;
 			Offset += BytesRead;
 			BytesToRead -= BytesRead;
 			Destination += BytesRead;
