[BuildFix] NonUnity build fix

#rnx

#Synced-CL 50389196

diff --git a/Engine/Source/Runtime/Core/Private/Android/AndroidPlatformFile.cpp b/Engine/Source/Runtime/Core/Private/Android/AndroidPlatformFile.cpp
--- a/Engine/Source/Runtime/Core/Private/Android/AndroidPlatformFile.cpp
+++ b/Engine/Source/Runtime/Core/Private/Android/AndroidPlatformFile.cpp
@@ -25,6 +25,7 @@
 #include <limits>
 #include <sys/mman.h>
 #include "Misc/ScopeLock.h"
+#include "Misc/ScopeExit.h"
 #include <Async/MappedFileHandle.h>
 
 #if PLATFORM_USE_PLATFORM_FILE_MANAGED_STORAGE_WRAPPER
diff --git a/Engine/Source/Runtime/Core/Private/Apple/ApplePlatformFile.cpp b/Engine/Source/Runtime/Core/Private/Apple/ApplePlatformFile.cpp
--- a/Engine/Source/Runtime/Core/Private/Apple/ApplePlatformFile.cpp
+++ b/Engine/Source/Runtime/Core/Private/Apple/ApplePlatformFile.cpp
@@ -22,6 +22,7 @@
 #include "HAL/PlatformFileCommon.h"
 #include "HAL/PlatformTime.h"
 #include "Misc/ScopeLock.h"
+#include "Misc/ScopeExit.h"
 #include "ProfilingDebugging/PlatformFileTrace.h"
 #include "HAL/DiskUtilizationTracker.h"
 #include "Templates/Function.h"
diff --git a/Engine/Source/Runtime/Core/Private/IOS/IOSPlatformFile.cpp b/Engine/Source/Runtime/Core/Private/IOS/IOSPlatformFile.cpp
--- a/Engine/Source/Runtime/Core/Private/IOS/IOSPlatformFile.cpp
+++ b/Engine/Source/Runtime/Core/Private/IOS/IOSPlatformFile.cpp
@@ -23,6 +23,7 @@
 #include "Misc/App.h"
 #include "Misc/Paths.h"
 #include "Misc/ScopeLock.h"
+#include "Misc/ScopeExit.h"
 #include "ProfilingDebugging/PlatformFileTrace.h"
 #include "HAL/DiskUtilizationTracker.h"
 
