Formalized FScopedCVar class.  Allows temporary setting of a CVar with a value, which is then unset.

-Introduced new EConsoleVariableFlags::ECVF_SetByTemp which is an array type.
- ECVF_SetByTemp is used by FScopedCVar to Set and Unset the variable.

#Synced-CL 50428254

diff --git a/Engine/Source/Runtime/Core/Private/HAL/ConsoleManager.cpp b/Engine/Source/Runtime/Core/Private/HAL/ConsoleManager.cpp
--- a/Engine/Source/Runtime/Core/Private/HAL/ConsoleManager.cpp
+++ b/Engine/Source/Runtime/Core/Private/HAL/ConsoleManager.cpp
@@ -757,7 +757,8 @@ constexpr bool IsArrayPriority(EConsoleVariableFlags Priority)
 		Priority == ECVF_SetByHotfix ||
 		Priority == ECVF_SetByDeviceProfile ||
 		Priority == ECVF_SetByScalability ||
-		Priority == ECVF_SetByPreview;
+		Priority == ECVF_SetByPreview ||
+		Priority == ECVF_SetByTemp;
 }
 
 template <class T>
diff --git a/Engine/Source/Runtime/Core/Public/HAL/IConsoleManager.h b/Engine/Source/Runtime/Core/Public/HAL/IConsoleManager.h
--- a/Engine/Source/Runtime/Core/Public/HAL/IConsoleManager.h
+++ b/Engine/Source/Runtime/Core/Public/HAL/IConsoleManager.h
@@ -171,8 +171,10 @@ enum EConsoleVariableFlags
 	ECVF_SetByCommandline =			0x0D000000,
 	// Used for high priority temporary debugging or operation modes
 	ECVF_SetByCode =				0x0E000000,
+	// Used for temporarily replacing a variable.  THIS IS AN ARRAY TYPE meaning multiple vales with this SetBy are stored in the history
+	ECVF_SetByTemp =                0x0F000000,
 	// Highest priority used via editor UI or or game/editor interactive console
-	ECVF_SetByConsole =				0x0F000000,
+	ECVF_SetByConsole =				0x10000000,
 
 
 	// ------------------------------------------------
@@ -194,6 +196,7 @@ enum EConsoleVariableFlags
 	op(Preview) \
 	op(Commandline) \
 	op(Code) \
+	op(Temp) \
 	op(Console)
 
 
diff --git a/Engine/Source/Runtime/Core/Public/Misc/ScopedCVar.h b/Engine/Source/Runtime/Core/Public/Misc/ScopedCVar.h
new file mode 100644
--- /dev/null
+++ b/Engine/Source/Runtime/Core/Public/Misc/ScopedCVar.h
@@ -0,0 +1,44 @@
+// Copyright Epic Games, Inc. All Rights Reserved.
+#pragma once
+
+#include "Containers/UnrealString.h"
+#include "HAL/IConsoleManager.h"
+#include "UObject/NameTypes.h"
+
+/**
+* FScopedCVar Allows a CVar to be temporarily set to a particular value.
+* Nested scopes still get popped off correctly since ECVF_SetByTemp is
+* an array type.
+*/
+
+template<typename T>
+class FScopedCVar
+{
+public:
+	FScopedCVar(FString InConsoleVariableName, T Value, FName Tag = "FScopedCVar")
+		: CachedVariableName(InConsoleVariableName)
+		, CachedTag(Tag)
+	{
+		IConsoleVariable* ConsoleVariable = IConsoleManager::Get().FindConsoleVariable(*CachedVariableName);
+		if (ensure(ConsoleVariable))
+		{
+			ConsoleVariable->Set(Value, ECVF_SetByTemp, Tag);
+		}
+	}
+
+	~FScopedCVar()
+	{
+		IConsoleVariable* ConsoleVariable = IConsoleManager::Get().FindConsoleVariable(*CachedVariableName);
+		if (ensure(ConsoleVariable))
+		{
+			ConsoleVariable->Unset(ECVF_SetByTemp, CachedTag);
+		}
+	}
+
+	FScopedCVar(const FScopedCVar&) = delete;
+	FScopedCVar& operator=(const FScopedCVar&) = delete;
+
+private:
+	FString CachedVariableName;
+	FName CachedTag;
+};
