Fixed FormatStringSan not correctly checking trailing arguments when the last specifier has more text after it.

#Synced-CL 50486641

diff --git a/Engine/Source/Runtime/Core/Public/String/FormatStringSan.h b/Engine/Source/Runtime/Core/Public/String/FormatStringSan.h
--- a/Engine/Source/Runtime/Core/Public/String/FormatStringSan.h
+++ b/Engine/Source/Runtime/Core/Public/String/FormatStringSan.h
@@ -131,17 +131,18 @@ namespace UE::Core::Private
 			{
 				using ArgPointeeType = std::remove_pointer_t<Arg>;
 
-				if (Fmt[0] == (CharType)'\0' && !bInsideFormatSpec)
-				{
-					return {EFormatStringSanStatus::NotEnoughSpecifiers, CurArgPos};
-				}
-
 				if (!bInsideFormatSpec)
 				{
 					while (*Fmt != (CharType)'\0' && *Fmt != (CharType)'%')
 					{
 						++Fmt;
 					}
+
+					if (Fmt[0] == (CharType)'\0')
+					{
+						return { EFormatStringSanStatus::NotEnoughSpecifiers, CurArgPos };
+					}
+
 					if (*Fmt == (CharType)'%')
 					{
 						++Fmt;
diff --git a/Engine/Source/Runtime/Core/Tests/String/FormatStringSanTest.cpp b/Engine/Source/Runtime/Core/Tests/String/FormatStringSanTest.cpp
--- a/Engine/Source/Runtime/Core/Tests/String/FormatStringSanTest.cpp
+++ b/Engine/Source/Runtime/Core/Tests/String/FormatStringSanTest.cpp
@@ -304,14 +304,12 @@ TEST_CASE_NAMED(FFormatStringValidatorTest, "System::Core::String::FormatStringS
 
 		STATIC_CHECK(UE_CHECK_FORMAT_STRING_ERR(E::Ok, TEXT("%-*.*d %f"), 4, 8, 42, 4.4));
 		STATIC_CHECK(UE_CHECK_FORMAT_STRING_ERR(E::Ok, TEXT("%*.*d %f"), 4, 8, 42, 4.4));
-		STATIC_CHECK(UE_CHECK_FORMAT_STRING_ERR(E::Ok, TEXT("%-*.9d $d"), 4, 42, 44));
-		STATIC_CHECK(UE_CHECK_FORMAT_STRING_ERR(E::Ok, TEXT("%*.9d $d"), 4, 42, 44));
+		STATIC_CHECK(UE_CHECK_FORMAT_STRING_ERR(E::Ok, TEXT("%-*.9d %d"), 4, 42, 44));
+		STATIC_CHECK(UE_CHECK_FORMAT_STRING_ERR(E::Ok, TEXT("%*.9d %d"), 4, 42, 44));
 		STATIC_CHECK(UE_CHECK_FORMAT_STRING_ERR(E::Ok, TEXT("%-*d %s"), 4, 42, TEXT("hi")));
 		STATIC_CHECK(UE_CHECK_FORMAT_STRING_ERR(E::Ok, TEXT("%*d %s"), 4, 42, TEXT("hi")));
 		STATIC_CHECK(UE_CHECK_FORMAT_STRING_ERR(E::Ok, TEXT("%.*f %s"), 4, 42.4, TEXT("hi")));
 
-		STATIC_CHECK(UE_CHECK_FORMAT_STRING_ERR(E::Ok, TEXT("Test extra arg '%s'."), TEXT("ok"), TEXT("hi")));
-
 		STATIC_CHECK(UE_CHECK_FORMAT_STRING_ERR(E::Ok, TEXT("a")));
 		STATIC_CHECK(UE_CHECK_FORMAT_STRING_ERR(E::Ok, TEXT("")));
 
@@ -341,12 +339,21 @@ TEST_CASE_NAMED(FFormatStringValidatorTest, "System::Core::String::FormatStringS
 		STATIC_CHECK(UE_CHECK_FORMAT_STRING_ERR(E::NotEnoughSpecifiers, TEXT(""), 1));
 		STATIC_CHECK(UE_CHECK_FORMAT_STRING_ERR(E::NotEnoughSpecifiers, TEXT(""), 1, 2));
 		STATIC_CHECK(UE_CHECK_FORMAT_STRING_ERR(E::NotEnoughSpecifiers, TEXT(""), 1, 2, 3));
+		STATIC_CHECK(UE_CHECK_FORMAT_STRING_ERR(E::Ok,                  TEXT("abcxyz")));
+		STATIC_CHECK(UE_CHECK_FORMAT_STRING_ERR(E::NotEnoughSpecifiers, TEXT("abcxyz"), 1));
+		STATIC_CHECK(UE_CHECK_FORMAT_STRING_ERR(E::NotEnoughSpecifiers, TEXT("abcxyz"), 1, 2));
+		STATIC_CHECK(UE_CHECK_FORMAT_STRING_ERR(E::NotEnoughSpecifiers, TEXT("abcxyz"), 1, 2, 3));
 
 		STATIC_CHECK(UE_CHECK_FORMAT_STRING_ERR(E::NotEnoughArguments,  TEXT("%d")));
 		STATIC_CHECK(UE_CHECK_FORMAT_STRING_ERR(E::Ok,                  TEXT("%d"), 1));
 		STATIC_CHECK(UE_CHECK_FORMAT_STRING_ERR(E::NotEnoughSpecifiers, TEXT("%d"), 1, 2));
 		STATIC_CHECK(UE_CHECK_FORMAT_STRING_ERR(E::NotEnoughSpecifiers, TEXT("%d"), 1, 2, 3));
 		STATIC_CHECK(UE_CHECK_FORMAT_STRING_ERR(E::NotEnoughSpecifiers, TEXT("%d"), 1, 2, 3, 4));
+		STATIC_CHECK(UE_CHECK_FORMAT_STRING_ERR(E::NotEnoughArguments,  TEXT("abc%dxyz")));
+		STATIC_CHECK(UE_CHECK_FORMAT_STRING_ERR(E::Ok,                  TEXT("abc%dxyz"), 1));
+		STATIC_CHECK(UE_CHECK_FORMAT_STRING_ERR(E::NotEnoughSpecifiers, TEXT("abc%dxyz"), 1, 2));
+		STATIC_CHECK(UE_CHECK_FORMAT_STRING_ERR(E::NotEnoughSpecifiers, TEXT("abc%dxyz"), 1, 2, 3));
+		STATIC_CHECK(UE_CHECK_FORMAT_STRING_ERR(E::NotEnoughSpecifiers, TEXT("abc%dxyz"), 1, 2, 3, 4));
 
 		STATIC_CHECK(UE_CHECK_FORMAT_STRING_ERR(E::NotEnoughArguments,  TEXT("%d %d")));
 		STATIC_CHECK(UE_CHECK_FORMAT_STRING_ERR(E::NotEnoughArguments,  TEXT("%d %d"), 1));
@@ -354,6 +361,12 @@ TEST_CASE_NAMED(FFormatStringValidatorTest, "System::Core::String::FormatStringS
 		STATIC_CHECK(UE_CHECK_FORMAT_STRING_ERR(E::NotEnoughSpecifiers, TEXT("%d %d"), 1, 2, 3));
 		STATIC_CHECK(UE_CHECK_FORMAT_STRING_ERR(E::NotEnoughSpecifiers, TEXT("%d %d"), 1, 2, 3, 4));
 		STATIC_CHECK(UE_CHECK_FORMAT_STRING_ERR(E::NotEnoughSpecifiers, TEXT("%d %d"), 1, 2, 3, 4, 5));
+		STATIC_CHECK(UE_CHECK_FORMAT_STRING_ERR(E::NotEnoughArguments,  TEXT("abc%d %dxyz")));
+		STATIC_CHECK(UE_CHECK_FORMAT_STRING_ERR(E::NotEnoughArguments,  TEXT("abc%d %dxyz"), 1));
+		STATIC_CHECK(UE_CHECK_FORMAT_STRING_ERR(E::Ok,                  TEXT("abc%d %dxyz"), 1, 2));
+		STATIC_CHECK(UE_CHECK_FORMAT_STRING_ERR(E::NotEnoughSpecifiers, TEXT("abc%d %dxyz"), 1, 2, 3));
+		STATIC_CHECK(UE_CHECK_FORMAT_STRING_ERR(E::NotEnoughSpecifiers, TEXT("abc%d %dxyz"), 1, 2, 3, 4));
+		STATIC_CHECK(UE_CHECK_FORMAT_STRING_ERR(E::NotEnoughSpecifiers, TEXT("abc%d %dxyz"), 1, 2, 3, 4, 5));
 
 		STATIC_CHECK(UE_CHECK_FORMAT_STRING_ERR(E::NotEnoughArguments,  TEXT("%d %d %d")));
 		STATIC_CHECK(UE_CHECK_FORMAT_STRING_ERR(E::NotEnoughArguments,  TEXT("%d %d %d"), 1));
@@ -362,6 +375,13 @@ TEST_CASE_NAMED(FFormatStringValidatorTest, "System::Core::String::FormatStringS
 		STATIC_CHECK(UE_CHECK_FORMAT_STRING_ERR(E::NotEnoughSpecifiers, TEXT("%d %d %d"), 1, 2, 3, 4));
 		STATIC_CHECK(UE_CHECK_FORMAT_STRING_ERR(E::NotEnoughSpecifiers, TEXT("%d %d %d"), 1, 2, 3, 4, 5));
 		STATIC_CHECK(UE_CHECK_FORMAT_STRING_ERR(E::NotEnoughSpecifiers, TEXT("%d %d %d"), 1, 2, 3, 4, 5, 6));
+		STATIC_CHECK(UE_CHECK_FORMAT_STRING_ERR(E::NotEnoughArguments,  TEXT("abc%d %d %dxyz")));
+		STATIC_CHECK(UE_CHECK_FORMAT_STRING_ERR(E::NotEnoughArguments,  TEXT("abc%d %d %dxyz"), 1));
+		STATIC_CHECK(UE_CHECK_FORMAT_STRING_ERR(E::NotEnoughArguments,  TEXT("abc%d %d %dxyz"), 1, 2));
+		STATIC_CHECK(UE_CHECK_FORMAT_STRING_ERR(E::Ok,                  TEXT("abc%d %d %dxyz"), 1, 2, 3));
+		STATIC_CHECK(UE_CHECK_FORMAT_STRING_ERR(E::NotEnoughSpecifiers, TEXT("abc%d %d %dxyz"), 1, 2, 3, 4));
+		STATIC_CHECK(UE_CHECK_FORMAT_STRING_ERR(E::NotEnoughSpecifiers, TEXT("abc%d %d %dxyz"), 1, 2, 3, 4, 5));
+		STATIC_CHECK(UE_CHECK_FORMAT_STRING_ERR(E::NotEnoughSpecifiers, TEXT("abc%d %d %dxyz"), 1, 2, 3, 4, 5, 6));
 	}
 }
 
