Add `USING_UNDEFINED_BEHAVIOR_SANITISER` using the `__has_feature` trick (Clang has supported it since 2018).

#Synced-CL 50530424

diff --git a/Engine/Source/Runtime/Core/Public/Clang/ClangPlatformCodeAnalysis.h b/Engine/Source/Runtime/Core/Public/Clang/ClangPlatformCodeAnalysis.h
--- a/Engine/Source/Runtime/Core/Public/Clang/ClangPlatformCodeAnalysis.h
+++ b/Engine/Source/Runtime/Core/Public/Clang/ClangPlatformCodeAnalysis.h
@@ -113,8 +113,9 @@
 	#define USING_THREAD_SANITISER 0
 #endif
 
-// Clang does not expose __has_feature(undefined_behavior_sanitizer) so this define comes directly from UBT.
-#ifndef USING_UNDEFINED_BEHAVIOR_SANITISER
+#if defined(__has_feature) && __has_feature(undefined_behavior_sanitizer)
+	#define USING_UNDEFINED_BEHAVIOR_SANITISER 1
+#else
 	#define USING_UNDEFINED_BEHAVIOR_SANITISER 0
 #endif
 
diff --git a/Engine/Source/Runtime/Core/Public/HAL/Platform.h b/Engine/Source/Runtime/Core/Public/HAL/Platform.h
--- a/Engine/Source/Runtime/Core/Public/HAL/Platform.h
+++ b/Engine/Source/Runtime/Core/Public/HAL/Platform.h
@@ -141,6 +141,14 @@
 	#define USING_HW_ADDRESS_SANITISER 0
 #endif
 
+#ifndef USING_THREAD_SANITISER
+	#define USING_THREAD_SANITISER 0
+#endif
+
+#ifndef USING_UNDEFINED_BEHAVIOR_SANITISER
+	#define USING_UNDEFINED_BEHAVIOR_SANITISER 0
+#endif
+
 #ifndef PLATFORM_HAS_ASAN_INCLUDE
 	#define PLATFORM_HAS_ASAN_INCLUDE __has_include(<sanitizer/asan_interface.h>)
 #endif
diff --git a/Engine/Source/Runtime/Core/Public/Misc/CoreMiscDefines.h b/Engine/Source/Runtime/Core/Public/Misc/CoreMiscDefines.h
--- a/Engine/Source/Runtime/Core/Public/Misc/CoreMiscDefines.h
+++ b/Engine/Source/Runtime/Core/Public/Misc/CoreMiscDefines.h
@@ -130,10 +130,6 @@
 	#define USING_INSTRUMENTATION 0
 #endif
 
-#ifndef USING_THREAD_SANITISER
-	#define USING_THREAD_SANITISER 0
-#endif
-
 #if USING_THREAD_SANITISER || USING_INSTRUMENTATION
 	#if !defined( TSAN_SAFE ) || !defined( TSAN_BEFORE ) || !defined( TSAN_AFTER ) || !defined( TSAN_ATOMIC )
 		#error Thread Sanitiser macros are not configured correctly for this platform
