Avoid assertion when running smoke tests when a test could not be run

#rnx
#Synced-CL 50536642

diff --git a/Engine/Source/Runtime/Core/Private/Misc/AutomationTest.cpp b/Engine/Source/Runtime/Core/Private/Misc/AutomationTest.cpp
--- a/Engine/Source/Runtime/Core/Private/Misc/AutomationTest.cpp
+++ b/Engine/Source/Runtime/Core/Private/Misc/AutomationTest.cpp
@@ -563,7 +563,7 @@ bool FAutomationTestFramework::RunSmokeTests()
 					
 					int32 RoleIndex = 0;  //always default to "local" role index.  Only used for multi-participant tests
 					StartTestByName( TestCommand, RoleIndex );
-					const bool CurTestSuccessful = StopTest(CurExecutionInfo);
+					const bool CurTestSuccessful = GIsAutomationTesting && StopTest(CurExecutionInfo);
 
 					bAllSuccessful = bAllSuccessful && CurTestSuccessful;
 				}
