[Trace Services] DeltaVarInt compression scheme
- Introduced DeltaVarInt as the default callstack compression scheme
- Renamed ZigZagAndVLE to Delta7bit

#rnx

#Synced-CL 50485454

diff --git a/Engine/Source/Developer/TraceServices/Private/Analyzers/CallstacksAnalysis.cpp b/Engine/Source/Developer/TraceServices/Private/Analyzers/CallstacksAnalysis.cpp
--- a/Engine/Source/Developer/TraceServices/Private/Analyzers/CallstacksAnalysis.cpp
+++ b/Engine/Source/Developer/TraceServices/Private/Analyzers/CallstacksAnalysis.cpp
@@ -16,7 +16,8 @@ enum ERoutes
 {
 	RouteId_Callstack,
 	RouteId_CallstackXORAndRLE,
-	RouteId_CallstackZigZagAndVLE
+	RouteId_CallstackDelta7bit,
+	RouteId_CallstackDeltaVarInt
 };
 	
 ////////////////////////////////////////////////////////////////////////////////
@@ -33,7 +34,9 @@ void FCallstacksAnalyzer::OnAnalysisBegin(const FOnAnalysisContext& Context)
 	FInterfaceBuilder& Builder = Context.InterfaceBuilder;
 	Builder.RouteEvent(RouteId_Callstack, "Memory", "CallstackSpec");
 	Builder.RouteEvent(RouteId_CallstackXORAndRLE, "Memory", "CallstackSpecXORAndRLE");
-	Builder.RouteEvent(RouteId_CallstackZigZagAndVLE, "Memory", "CallstackSpecZigZagAndVLE");
+	// NOTE: Update Event Tag to 'CallstackSpecDelta7bit' before UE 5.8 
+	Builder.RouteEvent(RouteId_CallstackDelta7bit, "Memory", "CallstackSpecZigZagAndVLE");
+	Builder.RouteEvent(RouteId_CallstackDeltaVarInt, "Memory", "CallstackSpecDeltaVarInt");
 }
 
 ////////////////////////////////////////////////////////////////////////////////
@@ -41,6 +44,28 @@ bool FCallstacksAnalyzer::OnEvent(uint16 RouteId, EStyle Style, const FOnEventCo
 {
 	LLM_SCOPE_BYNAME(TEXT("Insights/FCallstacksAnalyzer"));
 
+	const auto CompressionRoute = [&](auto UncompressFunction)
+	{
+		const TArrayReader<uint8>& CompressedData = Context.EventData.GetArray<uint8>("CompressedFrames");
+		if (const uint32 Id = Context.EventData.GetValue<uint32>("CallstackId"))
+		{
+			TStaticArray<uint64, 255> Frames;
+			uint32 TotalFrameCount = 0;
+			const uint32 FrameCount = UncompressFunction(
+				MakeArrayView(CompressedData.GetData(), CompressedData.Num()),
+				Frames,
+				TotalFrameCount
+			);
+
+			check(FrameCount <= 255);
+			if (TotalFrameCount != FrameCount)
+	        {
+				UE_LOG(LogTraceServices, Warning, TEXT("Callstack with Id=%u has %u frames, but it will be limited to %u frames!"), Id, TotalFrameCount, FrameCount);
+	        }
+			Provider->AddCallstack(Id, Frames.GetData(), uint8(FrameCount));
+		}
+	};
+	
 	switch (RouteId)
 	{
 		case RouteId_Callstack:
@@ -69,31 +94,19 @@ bool FCallstacksAnalyzer::OnEvent(uint16 RouteId, EStyle Style, const FOnEventCo
 
 		case RouteId_CallstackXORAndRLE:
 		{
-			const TArrayReader<uint8>& CompressedFrames = Context.EventData.GetArray<uint8>("CompressedFrames");
-			TStaticArray<uint64, UE_CALLSTACK_TRACE_MAX_FRAMES> Frames;
-			const uint32 FrameCount = FCallstackXORAndRLE::Uncompress(
-				MakeArrayView(CompressedFrames.GetData(), CompressedFrames.Num()),
-				Frames
-			);
-			if (const uint32 Id = Context.EventData.GetValue<uint32>("CallstackId"))
-			{
-				Provider->AddCallstack(Id, Frames.GetData(), uint8(FrameCount));
-			}
+			CompressionRoute(&FCallstackXORAndRLE::Uncompress);
 		}
 		break;
 
-		case RouteId_CallstackZigZagAndVLE:
+		case RouteId_CallstackDelta7bit:
 		{
-			const TArrayReader<uint8>& CompressedFrames = Context.EventData.GetArray<uint8>("CompressedFrames");
-			TStaticArray<uint64, UE_CALLSTACK_TRACE_MAX_FRAMES> Frames;
-			const uint32 FrameCount = FCallstackZigZagAndVLE::Uncompress(
-				MakeArrayView(CompressedFrames.GetData(), CompressedFrames.Num()),
-				Frames
-			);
-			if (const uint32 Id = Context.EventData.GetValue<uint32>("CallstackId"))
-			{
-				Provider->AddCallstack(Id, Frames.GetData(), uint8(FrameCount));
-			}
+			CompressionRoute(&FCallstackDelta7bit::Uncompress);
+		}
+		break;
+
+		case RouteId_CallstackDeltaVarInt:
+		{
+			CompressionRoute(&FCallstackDeltaVarInt::Uncompress);
 		}
 		break;
 	}
diff --git a/Engine/Source/Runtime/Core/Private/ProfilingDebugging/CallstackTrace.cpp b/Engine/Source/Runtime/Core/Private/ProfilingDebugging/CallstackTrace.cpp
--- a/Engine/Source/Runtime/Core/Private/ProfilingDebugging/CallstackTrace.cpp
+++ b/Engine/Source/Runtime/Core/Private/ProfilingDebugging/CallstackTrace.cpp
@@ -15,7 +15,8 @@ void	CallstackTrace_InitializeInternal();
 UE_TRACE_CHANNEL_DEFINE(CallstackChannel)
 UE_TRACE_EVENT_DEFINE(Memory, CallstackSpec)
 UE_TRACE_EVENT_DEFINE(Memory, CallstackSpecXORAndRLE)
-UE_TRACE_EVENT_DEFINE(Memory, CallstackSpecZigZagAndVLE)
+UE_TRACE_EVENT_DEFINE(Memory, CallstackSpecDelta7bit)
+UE_TRACE_EVENT_DEFINE(Memory, CallstackSpecDeltaVarInt)
 
 uint32 GCallStackTracingTlsSlotIndex = FPlatformTLS::InvalidTlsSlot;
 
@@ -42,123 +43,4 @@ void CallstackTrace_Initialize()
 	}();
 }
 
-#define UE_CALLSTACK_TRACE_TESTS 0
-#if UE_CALLSTACK_TRACE_TESTS
-
-#include "Algo/Compare.h"
-
-void Test_Compression()
-{
-	auto CheckXORAndRLE = [](TConstArrayView<uint64> Input)
-	{
-		uint8 Compressed[UE_CALLSTACK_TRACE_MAX_FRAMES*sizeof(uint64)];
-		const uint32 CompressedSize = FCallstackXORAndRLE::Compress(Input, Compressed);
-		const auto CompressedView = MakeArrayView(Compressed, CompressedSize);
-
-		TStaticArray<uint64, UE_CALLSTACK_TRACE_MAX_FRAMES> Uncompressed;
-		const uint32 UncompressedCount = FCallstackXORAndRLE::Uncompress(CompressedView, Uncompressed);
-		auto UncompressedView = MakeArrayView(Uncompressed.GetData(), UncompressedCount);
-
-		check(Algo::Compare(UncompressedView, Input));
-	};
-	
-	auto CheckZigZagAndVLE = [](TConstArrayView<uint64> Input)
-	{
-		uint8 Compressed[UE_CALLSTACK_TRACE_MAX_FRAMES*sizeof(uint64)];
-		const uint32 CompressedSize = FCallstackZigZagAndVLE::Compress(Input, Compressed);
-		const auto CompressedView = MakeArrayView(Compressed, CompressedSize);
-
-		TStaticArray<uint64, UE_CALLSTACK_TRACE_MAX_FRAMES> Uncompressed;
-		const uint32 UncompressedCount = FCallstackZigZagAndVLE::Uncompress(CompressedView, Uncompressed);
-		auto UncompressedView = MakeArrayView(Uncompressed.GetData(), UncompressedCount);
-
-		check(Algo::Compare(UncompressedView, Input));
-	};
-
-	uint64 Frames[] = {
-		0x00007ff693ab7b09,
-		0x00007ff693ab7cd2,
-		0x00007ff693aa9fe1,
-		0x00007ff693aaa0c8,
-		0x00007ff693abb0d0,
-		0x00007ff693abb0d0,
-		0x00007ffd61917344,
-		0x00007ffd639226b1,
-	};
-	uint8 FramesExpected[] = { 0x0 };
-	CheckXORAndRLE(Frames);
-	CheckZigZagAndVLE(Frames);
-	
-	uint64 Zeroes[] = {
-		0x0, 0x0, 0x0, 0x0, 0x0
-	};
-	uint8 ZeroesExpected[] = {
-		64, 64, 64, 64, 64
-	};
-	CheckXORAndRLE(Zeroes);
-    CheckZigZagAndVLE(Zeroes);
-	
-	uint64 SingleOne[] = {
-		0x1,
-	};
-	uint8 SingleOneExpected [] = {
-		63, 0x1
-	};
-	CheckXORAndRLE(SingleOne);
-	CheckZigZagAndVLE(SingleOne);
-	
-	uint64 SingleNormal[] = {
-		0x2d4f763d,
-	};
-	uint8 SingleNormalExpected [] = {
-		34, 0x3d, 0x76, 0x4f, 0x2d
-	};
-	CheckXORAndRLE(SingleNormal);
-	CheckZigZagAndVLE(SingleNormal);
-	
-	uint64 Mixed[] = {
-		0x000000002d4f829c,
-		0x00000001ed4f763d,
-		0xffffffffffffffff,
-		0x0000001800000000,
-		0x000000002d4f7680,
-		0x00000001ed4f7638,
-		0x0,
-		0x000000002d4f829c,
-	};
-	uint8 MixedExpected [] = {
-		34, 0x9c, 0x82, 0x4f, 0x2d, 
-		31, 0x3d, 0x76, 0x4f, 0xed, 0x01,
-		0,  0xff, 0xff, 0xff, 0xff, 0xff, 0xff, 0xff, 0xff,
-		27, 0x00, 0x00, 0x00, 0x00, 0x18,
-		34, 0x80, 0x76, 0x4f, 0x2d,
-		31, 0x38, 0x76, 0x4f, 0xed, 0x01,
-		64,
-		34, 0x9c, 0x82, 0x4f, 0x2d,
-	};
-	CheckXORAndRLE(Mixed);
-	CheckZigZagAndVLE(Mixed);
-
-	uint64 Random[] = {
-		0x85EA1249376518D5,
-		0x006B579412698957,
-		0x7B505F48F7D4B719,		
-	};
-	uint8 RandomExpected[] = {
-		0, 0xd5, 0x18, 0x65, 0x37, 0x49, 0x12, 0xEA, 0x85,
-		9, 0x57, 0x89, 0x69, 0x12, 0x94, 0x57, 0x6B,
-		1, 0x19, 0xB7, 0xD4, 0xF7, 0x48, 0x5F, 0x50, 0x7B
-	};
-	CheckXORAndRLE(Random);
-	CheckZigZagAndVLE(Random);
-}
-
-static bool bCompressionTestsWasRun = []()
-{
-	Test_Compression();
-	return true;
-}();
-
-#endif //UE_CALLSTACK_TRACE_TESTS
-
-#endif //UE_CALLSTACK_TRACE_ENABLED
\ No newline at end of file
+#endif //UE_CALLSTACK_TRACE_ENABLED
diff --git a/Engine/Source/Runtime/Core/Private/ProfilingDebugging/CallstackTracePrivate.h b/Engine/Source/Runtime/Core/Private/ProfilingDebugging/CallstackTracePrivate.h
--- a/Engine/Source/Runtime/Core/Private/ProfilingDebugging/CallstackTracePrivate.h
+++ b/Engine/Source/Runtime/Core/Private/ProfilingDebugging/CallstackTracePrivate.h
@@ -43,7 +43,12 @@ UE_TRACE_EVENT_BEGIN_EXTERN(Memory, CallstackSpecXORAndRLE, NoSync)
 	UE_TRACE_EVENT_FIELD(uint8[], CompressedFrames)
 UE_TRACE_EVENT_END()
 
-UE_TRACE_EVENT_BEGIN_EXTERN(Memory, CallstackSpecZigZagAndVLE, NoSync)
+UE_TRACE_EVENT_BEGIN_EXTERN(Memory, CallstackSpecDelta7bit, NoSync)
+	UE_TRACE_EVENT_FIELD(uint32, CallstackId)
+	UE_TRACE_EVENT_FIELD(uint8[], CompressedFrames)
+UE_TRACE_EVENT_END()
+
+UE_TRACE_EVENT_BEGIN_EXTERN(Memory, CallstackSpecDeltaVarInt, NoSync)
 	UE_TRACE_EVENT_FIELD(uint32, CallstackId)
 	UE_TRACE_EVENT_FIELD(uint8[], CompressedFrames)
 UE_TRACE_EVENT_END()
@@ -97,15 +102,15 @@ UE_TRACE_EVENT_END()
 					<< CallstackSpec.CallstackId(Id)
 					<< CallstackSpec.Frames(Entry.Frames, Entry.FrameCount);
 #else
-				uint8 Compressed[UE_CALLSTACK_TRACE_MAX_FRAMES*sizeof(uint64)];
-				const uint32 CompressedSize = FCallstackZigZagAndVLE::Compress(
+				uint8 Compressed[UE_CALLSTACK_TRACE_MAX_FRAMES * 9];
+				const uint32 CompressedSize = FCallstackDeltaVarInt::Compress(
 					MakeArrayView(Entry.Frames, Entry.FrameCount),
 					Compressed
 				);
 				
-				UE_TRACE_LOG(Memory, CallstackSpecZigZagAndVLE, CallstackChannel)
-					<< CallstackSpecZigZagAndVLE.CallstackId(Id)
-					<< CallstackSpecZigZagAndVLE.CompressedFrames(Compressed, CompressedSize);
+				UE_TRACE_LOG(Memory, CallstackSpecDeltaVarInt, CallstackChannel)
+					<< CallstackSpecDeltaVarInt.CallstackId(Id)
+					<< CallstackSpecDeltaVarInt.CompressedFrames(Compressed, CompressedSize);
 #endif //UE_CALLSTACK_TRACE_UNCOMPRESSED
 			}
 
diff --git a/Engine/Source/Runtime/Core/Public/ProfilingDebugging/CallstackTrace.h b/Engine/Source/Runtime/Core/Public/ProfilingDebugging/CallstackTrace.h
--- a/Engine/Source/Runtime/Core/Public/ProfilingDebugging/CallstackTrace.h
+++ b/Engine/Source/Runtime/Core/Public/ProfilingDebugging/CallstackTrace.h
@@ -6,6 +6,7 @@
 #include "HAL/PlatformTLS.h"
 #include "Containers/ArrayView.h"
 #include "Containers/StringFwd.h"
+#include "Serialization/VarInt.h"
 
 ////////////////////////////////////////////////////////////////////////////////
 #if !defined(UE_CALLSTACK_TRACE_ENABLED)
@@ -121,7 +122,7 @@ inline bool  CallstackTrace_ResolveFullCallStack() { return true; }
 
 /**
  * Simple run length compression of arrays of addresses (e.g stack frames).
- * Needs to be available regardless if tracing is enabled to allow analysis to decompress.
+ * Needs to be available regardless if tracing is enabled to allow compress/decompress analysis.
  */
 namespace FCallstackXORAndRLE
 {
@@ -154,11 +155,10 @@ namespace FCallstackXORAndRLE
 		return CompressedBytes;
 	}
 
-	inline uint32 Uncompress(TConstArrayView<uint8> Compressed, TArrayView<uint64> Addresses)
+	inline uint32 Uncompress(TConstArrayView<uint8> Compressed, TArrayView<uint64> Addresses, uint32& OutTotalAddressCount)
 	{
 		const int32 Bytes = Compressed.Num();
 		int32 ByteIndex = 0;
-		uint32 AddressCount = 0;
 		uint64 LastAddress = 0;
 		while (ByteIndex < Bytes)
 		{
@@ -176,28 +176,28 @@ namespace FCallstackXORAndRLE
 					BitsShift += 8;
 				}
 			}
-			
-			if (AddressCount)
+			OutAddress = LastAddress ^ OutAddress;
+			LastAddress = OutAddress;
+
+			if (OutTotalAddressCount < (uint32)Addresses.Num())
 			{
-				OutAddress = LastAddress ^ OutAddress;
+			    Addresses[OutTotalAddressCount] = OutAddress;
 			}
-
-			LastAddress = OutAddress;
-			Addresses[AddressCount++] = OutAddress;
+			OutTotalAddressCount++;
 		}
-		return AddressCount;
+		return FMath::Min(OutTotalAddressCount,  (uint32)Addresses.Num());
 	}
 } // namespace FCallstackXORAndRLE
 
 /**
- * ZigZag encoding with Variable-length integer compression of arrays of addresses (e.g stack frames).
- * Needs to be available regardless if tracing is enabled to allow analysis to decompress.
+ * Delta + ZigZag encoding with 7bit compression of arrays of addresses (e.g stack frames).
+ * Needs to be available regardless if tracing is enabled to allow compress/decompress analysis.
  */
-namespace FCallstackZigZagAndVLE
+namespace FCallstackDelta7bit
 {
 	inline uint32 Compress(TConstArrayView<uint64> Addresses, TArrayView<uint8> Compressed)
 	{
-		check(Addresses.Num() * 8 <= Compressed.Num());
+		check(Addresses.Num() * 10 <= Compressed.Num());
 		uint32 CompressedBytes = 0;
 		uint64 LastAddress = 0;
 		for(uint8 Index = 0; Index < Addresses.Num(); ++Index)
@@ -218,11 +218,10 @@ namespace FCallstackZigZagAndVLE
 		return CompressedBytes;
 	}
 
-	inline uint32 Uncompress(TConstArrayView<uint8> Compressed, TArrayView<uint64> Addresses)
+	inline uint32 Uncompress(TConstArrayView<uint8> Compressed, TArrayView<uint64> Addresses, uint32& OutTotalAddressCount)
 	{
 		const int32 Bytes = Compressed.Num();
 		int32 ByteIndex = 0;
-		uint32 AddressCount = 0;
 		uint64 LastAddress = 0;
 		while (ByteIndex < Bytes)
 		{
@@ -240,8 +239,57 @@ namespace FCallstackZigZagAndVLE
 			uint64 Address = LastAddress + Delta;
 			LastAddress = Address;
 			
-			Addresses[AddressCount++] = Address;
+			if (OutTotalAddressCount < (uint32)Addresses.Num())
+			{
+			    Addresses[OutTotalAddressCount] = Address;
+			}
+			OutTotalAddressCount++;
+		}
+		return FMath::Min(OutTotalAddressCount,  (uint32)Addresses.Num());
+	}
+} // namespace FCallstackDelta7bit
+
+/**
+ * Delta + ZigZag encoding with VarInt compression of arrays of addresses (e.g stack frames).
+ * Needs to be available regardless if tracing is enabled to allow compress/decompress analysis.
+ */
+namespace FCallstackDeltaVarInt
+{
+	inline uint32 Compress(TConstArrayView<uint64> Addresses, TArrayView<uint8> Compressed)
+	{
+		check(Addresses.Num() * 9 <= Compressed.Num());
+		uint8* CompressedPtr = Compressed.GetData();
+		uint64 LastAddress = 0;
+		for(uint8 Index = 0; Index < Addresses.Num(); ++Index)
+		{
+			const uint64 Address = Addresses[Index];
+			int64 Delta = int64(Address - LastAddress);
+			LastAddress = Address;
+
+			uint32 BytesCompressed = WriteVarInt(Delta, CompressedPtr);
+			CompressedPtr += BytesCompressed;
+		}
+		return uint32(CompressedPtr - Compressed.GetData());
+	}
+
+	inline uint32 Uncompress(TConstArrayView<uint8> Compressed, TArrayView<uint64> Addresses, uint32& OutTotalAddressCount)
+	{
+		const uint8* CompressedPtr = Compressed.GetData();
+		const uint8* EndPtr = CompressedPtr + Compressed.Num();
+		uint64 Address = 0;
+		while (CompressedPtr < EndPtr)
+		{
+			uint32 BytesFound = 0;
+			int64 Delta = ReadVarInt(CompressedPtr, BytesFound);
+			CompressedPtr += BytesFound;
+			Address += Delta;
+
+			if (OutTotalAddressCount < (uint32)Addresses.Num())
+			{
+			    Addresses[OutTotalAddressCount] = Address;
+			}
+			OutTotalAddressCount++;
 		}
-		return AddressCount;
+		return FMath::Min(OutTotalAddressCount,  (uint32)Addresses.Num());
 	}
-} // namespace FCallstackZigZag
+} // namespace FCallstackDeltaVarInt
\ No newline at end of file
