Rename UE4_DO_ROOT_PRIVILEGE_CHECK to UE_DO_ROOT_PRIVILEGE_CHECK
Rename UE4_PLATFORM_SANITY_CHECK_OS_ALLOCATIONS to UE_PLATFORM_SANITY_CHECK_OS_ALLOCATIONS
Remove UE4_PLATFORM_REDUCE_NUMBER_OF_MAPS

#Synced-CL 50505807

diff --git a/Engine/Source/Runtime/Core/Private/Apple/ApplePlatformMemory.cpp b/Engine/Source/Runtime/Core/Private/Apple/ApplePlatformMemory.cpp
--- a/Engine/Source/Runtime/Core/Private/Apple/ApplePlatformMemory.cpp
+++ b/Engine/Source/Runtime/Core/Private/Apple/ApplePlatformMemory.cpp
@@ -484,10 +484,9 @@ bool FApplePlatformMemory::PageProtect(void* const Ptr, const SIZE_T Size, const
 #ifndef MALLOC_LEAKDETECTION
 	#define MALLOC_LEAKDETECTION 0
 #endif
-#define UE4_PLATFORM_REDUCE_NUMBER_OF_MAPS 0
 
 // check bookkeeping info against the passed in parameters in Debug and Development (the latter only in games and servers. also, only if leak detection is disabled, otherwise things are very slow)
-#define UE4_PLATFORM_SANITY_CHECK_OS_ALLOCATIONS			(UE_BUILD_DEBUG || (UE_BUILD_DEVELOPMENT && (UE_GAME || UE_SERVER) && !MALLOC_LEAKDETECTION))
+#define UE_PLATFORM_SANITY_CHECK_OS_ALLOCATIONS			(UE_BUILD_DEBUG || (UE_BUILD_DEVELOPMENT && (UE_GAME || UE_SERVER) && !MALLOC_LEAKDETECTION))
 
 /** This structure is stored in the page after each OS allocation and checks that its properties are valid on Free. Should be less than the page size (4096 on all platforms we support) */
 struct FOSAllocationDescriptor
@@ -522,7 +521,11 @@ void* _Nullable FApplePlatformMemory::BinnedAllocFromOS(SIZE_T Size)
 	// Binned expects OS allocations to be BinnedPageSize-aligned, and that page is at least 64KB. mmap() alone cannot do this, so carve out the needed chunks.
 	const SIZE_T ExpectedAlignment = FPlatformMemory::GetConstants().BinnedPageSize;
 	// Descriptor is only used if we're sanity checking. However, #ifdef'ing its use would make the code more fragile. Size needs to be at least one page.
-	const SIZE_T DescriptorSize = (UE4_PLATFORM_SANITY_CHECK_OS_ALLOCATIONS != 0) ? OSPageSize : 0;
+#if UE_PLATFORM_SANITY_CHECK_OS_ALLOCATIONS
+	const SIZE_T DescriptorSize = OSPageSize;
+#else
+	const SIZE_T DescriptorSize = 0;
+#endif
 
 	SIZE_T ActualSizeMapped = SizeInWholePages + ExpectedAlignment;
 
@@ -548,21 +551,17 @@ void* _Nullable FApplePlatformMemory::BinnedAllocFromOS(SIZE_T Size)
 		SIZE_T SizeToNextAlignedPointer = ExpectedAlignment - Offset;
 		void* AlignedPointer = reinterpret_cast<void*>(reinterpret_cast<SIZE_T>(Pointer) + SizeToNextAlignedPointer);
 
-		// do not unmap if we're trying to reduce the number of distinct maps, since holes prevent the Linux kernel from coalescing two adjoining mmap()s into a single VMA
-		if (!UE4_PLATFORM_REDUCE_NUMBER_OF_MAPS)
+		// unmap the part before
+		if (munmap(Pointer, SizeToNextAlignedPointer) != 0)
 		{
-			// unmap the part before
-			if (munmap(Pointer, SizeToNextAlignedPointer) != 0)
-			{
-				FPlatformMemory::OnOutOfMemory(SizeToNextAlignedPointer, ExpectedAlignment);
-				// unreachable
-				return nullptr;
-			}
-
-			// account for reduced mmaped size
-			ActualSizeMapped -= SizeToNextAlignedPointer;
+			FPlatformMemory::OnOutOfMemory(SizeToNextAlignedPointer, ExpectedAlignment);
+			// unreachable
+			return nullptr;
 		}
 
+		// account for reduced mmaped size
+		ActualSizeMapped -= SizeToNextAlignedPointer;
+
 		// now, make it appear as if we initially got the allocation right
 		Pointer = AlignedPointer;
 	}
@@ -571,21 +570,17 @@ void* _Nullable FApplePlatformMemory::BinnedAllocFromOS(SIZE_T Size)
 	// or we already got rid of the extra memory that was allocated in the front.
 	checkf((reinterpret_cast<SIZE_T>(Pointer) % ExpectedAlignment) == 0, TEXT("BinnedAllocFromOS(): Internal error: did not align the pointer as expected."));
 
-	// do not unmap if we're trying to reduce the number of distinct maps, since holes prevent the kernel from coalescing two adjoining mmap()s into a single VMA
-	if (!UE4_PLATFORM_REDUCE_NUMBER_OF_MAPS)
-	{
-		// Now unmap the tail only, if any, but leave just enough space for the descriptor
-		void* TailPtr = reinterpret_cast<void*>(reinterpret_cast<SIZE_T>(Pointer) + SizeInWholePages + DescriptorSize);
-		SSIZE_T TailSize = ActualSizeMapped - SizeInWholePages - DescriptorSize;
+	// Now unmap the tail only, if any, but leave just enough space for the descriptor
+	void* TailPtr = reinterpret_cast<void*>(reinterpret_cast<SIZE_T>(Pointer) + SizeInWholePages + DescriptorSize);
+	SSIZE_T TailSize = ActualSizeMapped - SizeInWholePages - DescriptorSize;
 
-		if (LIKELY(TailSize > 0))
+	if (LIKELY(TailSize > 0))
+	{
+		if (munmap(TailPtr, TailSize) != 0)
 		{
-			if (munmap(TailPtr, TailSize) != 0)
-			{
-				FPlatformMemory::OnOutOfMemory(TailSize, ExpectedAlignment);
-				// unreachable
-				return nullptr;
-			}
+			FPlatformMemory::OnOutOfMemory(TailSize, ExpectedAlignment);
+			// unreachable
+			return nullptr;
 		}
 	}
 
@@ -594,16 +589,8 @@ void* _Nullable FApplePlatformMemory::BinnedAllocFromOS(SIZE_T Size)
 	{
 		FOSAllocationDescriptor* AllocDescriptor = reinterpret_cast<FOSAllocationDescriptor*>(reinterpret_cast<SIZE_T>(Pointer) + Size);
 		AllocDescriptor->Magic = FOSAllocationDescriptor::MagicType::Marker;
-		if (!UE4_PLATFORM_REDUCE_NUMBER_OF_MAPS)
-		{
-			AllocDescriptor->PointerToUnmap = Pointer;
-			AllocDescriptor->SizeToUnmap = SizeInWholePages + DescriptorSize;
-		}
-		else
-		{
-			AllocDescriptor->PointerToUnmap = PointerWeGotFromMMap;
-			AllocDescriptor->SizeToUnmap = ActualSizeMapped;
-		}
+		AllocDescriptor->PointerToUnmap = Pointer;
+		AllocDescriptor->SizeToUnmap = SizeInWholePages + DescriptorSize;
 		AllocDescriptor->OriginalSizeAsPassed = Size;
 	}
 
@@ -630,7 +617,7 @@ void FApplePlatformMemory::BinnedFreeToOS(void* Ptr, SIZE_T Size)
 	static SIZE_T OSPageSize = FPlatformMemory::GetConstants().PageSize;
 	SIZE_T SizeInWholePages = (Size % OSPageSize) ? (Size + OSPageSize - (Size % OSPageSize)) : Size;
 
-	if (UE4_PLATFORM_SANITY_CHECK_OS_ALLOCATIONS)
+	#if (UE_PLATFORM_SANITY_CHECK_OS_ALLOCATIONS)
 	{
 		const SIZE_T DescriptorSize = OSPageSize;
 
@@ -645,23 +632,19 @@ void FApplePlatformMemory::BinnedFreeToOS(void* Ptr, SIZE_T Size)
 		void* PointerToUnmap = AllocDescriptor->PointerToUnmap;
 		SIZE_T SizeToUnmap = AllocDescriptor->SizeToUnmap;
 
-		// do checks, from most to least serious
-		if (UE4_PLATFORM_SANITY_CHECK_OS_ALLOCATIONS != 0)
+		// this check only makes sense when we're not reducing number of maps, since the pointer will have to be different.
+		if (UNLIKELY(PointerToUnmap != Ptr || SizeToUnmap != SizeInWholePages + DescriptorSize))
 		{
-			// this check only makes sense when we're not reducing number of maps, since the pointer will have to be different.
-			if (UNLIKELY(PointerToUnmap != Ptr || SizeToUnmap != SizeInWholePages + DescriptorSize))
-			{
-				UE_LOG(LogHAL, Fatal, TEXT("BinnedFreeToOS(): info mismatch: descriptor ptr: %p, size %llu, but our pointer is %p and size %llu."), PointerToUnmap, SizeToUnmap, AllocDescriptor, (uint64)(SizeInWholePages + DescriptorSize));
-				// unreachable
-				return;
-			}
-
-			if (UNLIKELY(AllocDescriptor->OriginalSizeAsPassed != Size))
-			{
-				UE_LOG(LogHAL, Fatal, TEXT("BinnedFreeToOS(): info mismatch: descriptor original size %llu, our size is %llu for pointer %p"), AllocDescriptor->OriginalSizeAsPassed, Size, Ptr);
-				// unreachable
-				return;
-			}
+			UE_LOG(LogHAL, Fatal, TEXT("BinnedFreeToOS(): info mismatch: descriptor ptr: %p, size %llu, but our pointer is %p and size %llu."), PointerToUnmap, SizeToUnmap, AllocDescriptor, (uint64)(SizeInWholePages + DescriptorSize));
+			// unreachable
+			return;
+		}
+
+		if (UNLIKELY(AllocDescriptor->OriginalSizeAsPassed != Size))
+		{
+			UE_LOG(LogHAL, Fatal, TEXT("BinnedFreeToOS(): info mismatch: descriptor original size %llu, our size is %llu for pointer %p"), AllocDescriptor->OriginalSizeAsPassed, Size, Ptr);
+			// unreachable
+			return;
 		}
 
 		AllocDescriptor = nullptr;	// just so no one touches it
@@ -672,7 +655,7 @@ void FApplePlatformMemory::BinnedFreeToOS(void* Ptr, SIZE_T Size)
 			// unreachable
 		}
 	}
-	else
+	#else
 	{
 		if (UNLIKELY(munmap(Ptr, SizeInWholePages) != 0))
 		{
@@ -680,6 +663,7 @@ void FApplePlatformMemory::BinnedFreeToOS(void* Ptr, SIZE_T Size)
 			// unreachable
 		}
 	}
+	#endif
 #else
 	LLM_IF_ENABLED(FLowLevelMemTracker::Get().OnLowLevelFree(ELLMTracker::Platform, Ptr));
 	if (munmap(Ptr, Size) != 0)
diff --git a/Engine/Source/Runtime/Core/Private/Unix/UnixPlatformMemory.cpp b/Engine/Source/Runtime/Core/Private/Unix/UnixPlatformMemory.cpp
--- a/Engine/Source/Runtime/Core/Private/Unix/UnixPlatformMemory.cpp
+++ b/Engine/Source/Runtime/Core/Private/Unix/UnixPlatformMemory.cpp
@@ -31,19 +31,19 @@
 #include "GenericPlatform/OSAllocationPool.h"
 #include "Misc/ScopeLock.h"
 
-// on 64 bit Linux, it is easier to run out of vm.max_map_count than of other limits. Due to that, trade VIRT (address space) size for smaller amount of distinct mappings
-// by not leaving holes between them (kernel will coalesce the adjoining mappings into a single one)
-// Disable by default as this causes large wasted virtual memory areas as we've to cut out 64k aligned pointers from a larger memory area then requested but then leave them mmap'd
-#define UE4_PLATFORM_REDUCE_NUMBER_OF_MAPS					(0 && PLATFORM_UNIX)
-
+#ifdef UE4_DO_ROOT_PRIVILEGE_CHECK
+	// UE_DEPRECATED(5.8)
+	#warning "UE4_DO_ROOT_PRIVILEGE_CHECK has been renamed UE_DO_ROOT_PRIVILEGE_CHECK"
+	#define UE_DO_ROOT_PRIVILEGE_CHECK UE4_DO_ROOT_PRIVILEGE_CHECK
+#endif
 // do not do a root privilege check on non-x86-64 platforms (assume an embedded device)
-#ifndef UE4_DO_ROOT_PRIVILEGE_CHECK
+#ifndef UE_DO_ROOT_PRIVILEGE_CHECK
   #if defined(_M_X64) || defined(__x86_64__) || defined (__amd64__)
-    #define UE4_DO_ROOT_PRIVILEGE_CHECK	 1
+#define UE_DO_ROOT_PRIVILEGE_CHECK	 1
   #else
-    #define UE4_DO_ROOT_PRIVILEGE_CHECK	 0
+    #define UE_DO_ROOT_PRIVILEGE_CHECK	 0
   #endif // defined(_M_X64) || defined(__x86_64__) || defined (__amd64__) 
-#endif // ifndef UE4_DO_ROOT_PRIVILEGE_CHECK
+#endif // ifndef UE_DO_ROOT_PRIVILEGE_CHECK
 
 // Set to 1 if we should try to use /proc/self/smaps_rollup in FUnixPlatformMemory::GetExtendedStats().
 // There is a potential tradeoff in that smaps_rollup appears to be quite a bit faster in
@@ -138,7 +138,7 @@ class FMalloc* FUnixPlatformMemory::BaseAllocator()
 		return Allocator;
 	}
 
-#if UE4_DO_ROOT_PRIVILEGE_CHECK && !IS_PROGRAM
+#if UE_DO_ROOT_PRIVILEGE_CHECK && !IS_PROGRAM
 	// This function gets executed very early, way before main() (because global constructors will allocate memory).
 	// This makes it ideal, if unobvious, place for a root privilege check.
 	if (geteuid() == 0)
@@ -148,7 +148,7 @@ class FMalloc* FUnixPlatformMemory::BaseAllocator()
 		// unreachable
 		return nullptr;
 	}
-#endif // UE4_DO_ROOT_PRIVILEGE_CHECK && !IS_PROGRAM
+#endif // UE_DO_ROOT_PRIVILEGE_CHECK && !IS_PROGRAM
 
 #if UE_USE_MALLOC_REPLAY_PROXY
 	bool bAddReplayProxy = false;
@@ -403,7 +403,7 @@ static void MarkMappedMemoryMergable(void* Pointer, SIZE_T Size)
 	#define MALLOC_LEAKDETECTION 0
 #endif
 // check bookkeeping info against the passed in parameters in Debug and Development (the latter only in games and servers. also, only if leak detection is disabled, otherwise things are very slow)
-#define UE4_PLATFORM_SANITY_CHECK_OS_ALLOCATIONS			(UE_BUILD_DEBUG || (UE_BUILD_DEVELOPMENT && (UE_GAME || UE_SERVER) && !MALLOC_LEAKDETECTION))
+#define UE_PLATFORM_SANITY_CHECK_OS_ALLOCATIONS			(UE_BUILD_DEBUG || (UE_BUILD_DEVELOPMENT && (UE_GAME || UE_SERVER) && !MALLOC_LEAKDETECTION))
 
 /** This structure is stored in the page after each OS allocation and checks that its properties are valid on Free. Should be less than the page size (4096 on all platforms we support) */
 struct FOSAllocationDescriptor
@@ -445,7 +445,11 @@ void* FUnixPlatformMemory::BinnedAllocFromOS(SIZE_T Size)
 	// Binned expects OS allocations to be BinnedPageSize-aligned, and that page is at least 64KB. mmap() alone cannot do this, so carve out the needed chunks.
 	const SIZE_T ExpectedAlignment = FPlatformMemory::GetConstants().BinnedPageSize;
 	// Descriptor is only used if we're sanity checking. However, #ifdef'ing its use would make the code more fragile. Size needs to be at least one page.
-	const SIZE_T DescriptorSize = (UE4_PLATFORM_REDUCE_NUMBER_OF_MAPS != 0 || UE4_PLATFORM_SANITY_CHECK_OS_ALLOCATIONS != 0) ? OSPageSize : 0;
+#if (UE_PLATFORM_SANITY_CHECK_OS_ALLOCATIONS != 0)
+	const SIZE_T DescriptorSize = OSPageSize;
+#else
+	const SIZE_T DescriptorSize = 0;
+#endif
 
 	SIZE_T ActualSizeMapped = SizeInWholePages + ExpectedAlignment;
 
@@ -471,21 +475,17 @@ void* FUnixPlatformMemory::BinnedAllocFromOS(SIZE_T Size)
 		SIZE_T SizeToNextAlignedPointer = ExpectedAlignment - Offset;
 		void* AlignedPointer = reinterpret_cast<void*>(reinterpret_cast<SIZE_T>(Pointer) + SizeToNextAlignedPointer);
 
-		// do not unmap if we're trying to reduce the number of distinct maps, since holes prevent the Linux kernel from coalescing two adjoining mmap()s into a single VMA
-		if (!UE4_PLATFORM_REDUCE_NUMBER_OF_MAPS)
+		// unmap the part before
+		if (munmap(Pointer, SizeToNextAlignedPointer) != 0)
 		{
-			// unmap the part before
-			if (munmap(Pointer, SizeToNextAlignedPointer) != 0)
-			{
-				FPlatformMemory::OnOutOfMemory(SizeToNextAlignedPointer, ExpectedAlignment);
-				// unreachable
-				return nullptr;
-			}
-
-			// account for reduced mmaped size
-			ActualSizeMapped -= SizeToNextAlignedPointer;
+			FPlatformMemory::OnOutOfMemory(SizeToNextAlignedPointer, ExpectedAlignment);
+			// unreachable
+			return nullptr;
 		}
 
+		// account for reduced mmaped size
+		ActualSizeMapped -= SizeToNextAlignedPointer;
+
 		// now, make it appear as if we initially got the allocation right
 		Pointer = AlignedPointer;
 	}
@@ -496,41 +496,31 @@ void* FUnixPlatformMemory::BinnedAllocFromOS(SIZE_T Size)
 	// or we already got rid of the extra memory that was allocated in the front.
 	checkf((reinterpret_cast<SIZE_T>(Pointer) % ExpectedAlignment) == 0, TEXT("BinnedAllocFromOS(): Internal error: did not align the pointer as expected."));
 
-	// do not unmap if we're trying to reduce the number of distinct maps, since holes prevent the Linux kernel from coalescing two adjoining mmap()s into a single VMA
-	if (!UE4_PLATFORM_REDUCE_NUMBER_OF_MAPS)
-	{
-		// Now unmap the tail only, if any, but leave just enough space for the descriptor
-		void* TailPtr = reinterpret_cast<void*>(reinterpret_cast<SIZE_T>(Pointer) + SizeInWholePages + DescriptorSize);
-		SIZE_T TailSize = ActualSizeMapped - SizeInWholePages - DescriptorSize;
+	// Now unmap the tail only, if any, but leave just enough space for the descriptor
+	void* TailPtr = reinterpret_cast<void*>(reinterpret_cast<SIZE_T>(Pointer) + SizeInWholePages + DescriptorSize);
+	SIZE_T TailSize = ActualSizeMapped - SizeInWholePages - DescriptorSize;
 
-		if (LIKELY(TailSize > 0))
+	if (LIKELY(TailSize > 0))
+	{
+		if (munmap(TailPtr, TailSize) != 0)
 		{
-			if (munmap(TailPtr, TailSize) != 0)
-			{
-				FPlatformMemory::OnOutOfMemory(TailSize, ExpectedAlignment);
-				// unreachable
-				return nullptr;
-			}
+			FPlatformMemory::OnOutOfMemory(TailSize, ExpectedAlignment);
+			// unreachable
+			return nullptr;
 		}
 	}
 
+#if (UE_PLATFORM_SANITY_CHECK_OS_ALLOCATIONS != 0)
 	// we're done with this allocation, fill in the descriptor with the info
 	if (LIKELY(DescriptorSize > 0))
 	{
 		FOSAllocationDescriptor* AllocDescriptor = reinterpret_cast<FOSAllocationDescriptor*>(reinterpret_cast<SIZE_T>(Pointer) + Size);
 		AllocDescriptor->Magic = FOSAllocationDescriptor::MagicType::Marker;
-		if (!UE4_PLATFORM_REDUCE_NUMBER_OF_MAPS)
-		{
-			AllocDescriptor->PointerToUnmap = Pointer;
-			AllocDescriptor->SizeToUnmap = SizeInWholePages + DescriptorSize;
-		}
-		else
-		{
-			AllocDescriptor->PointerToUnmap = PointerWeGotFromMMap;
-			AllocDescriptor->SizeToUnmap = ActualSizeMapped;
-		}
+		AllocDescriptor->PointerToUnmap = Pointer;
+		AllocDescriptor->SizeToUnmap = SizeInWholePages + DescriptorSize;
 		AllocDescriptor->OriginalSizeAsPassed = Size;
 	}
+#endif
 
 	LLM_IF_ENABLED(FLowLevelMemTracker::Get().OnLowLevelAlloc(ELLMTracker::Platform, Pointer, Size));
 	UE::FForkPageProtector::Get().AddMemoryRegion(Pointer, Size);
@@ -545,7 +535,7 @@ void FUnixPlatformMemory::BinnedFreeToOS(void* Ptr, SIZE_T Size)
 	static SIZE_T OSPageSize = FPlatformMemory::GetConstants().PageSize;
 	SIZE_T SizeInWholePages = (Size % OSPageSize) ? (Size + OSPageSize - (Size % OSPageSize)) : Size;
 
-	if (UE4_PLATFORM_REDUCE_NUMBER_OF_MAPS || UE4_PLATFORM_SANITY_CHECK_OS_ALLOCATIONS)
+	#if (UE_PLATFORM_SANITY_CHECK_OS_ALLOCATIONS)
 	{
 		const SIZE_T DescriptorSize = OSPageSize;
 
@@ -562,26 +552,18 @@ void FUnixPlatformMemory::BinnedFreeToOS(void* Ptr, SIZE_T Size)
 
 		UE::FForkPageProtector::Get().FreeMemoryRegion(PointerToUnmap);
 
-		// do checks, from most to least serious
-		if (UE4_PLATFORM_SANITY_CHECK_OS_ALLOCATIONS != 0)
+		if (UNLIKELY(PointerToUnmap != Ptr || SizeToUnmap != SizeInWholePages + DescriptorSize))
 		{
-			// this check only makes sense when we're not reducing number of maps, since the pointer will have to be different.
-			if (UE4_PLATFORM_REDUCE_NUMBER_OF_MAPS == 0)
-			{
-				if (UNLIKELY(PointerToUnmap != Ptr || SizeToUnmap != SizeInWholePages + DescriptorSize))
-				{
-					UE_LOG(LogHAL, Fatal, TEXT("BinnedFreeToOS(): info mismatch: descriptor ptr: %p, size %llu, but our pointer is %p and size %llu."), PointerToUnmap, SizeToUnmap, AllocDescriptor, (uint64)(SizeInWholePages + DescriptorSize));
-					// unreachable
-					return;
-				}
-			}
+			UE_LOG(LogHAL, Fatal, TEXT("BinnedFreeToOS(): info mismatch: descriptor ptr: %p, size %llu, but our pointer is %p and size %llu."), PointerToUnmap, SizeToUnmap, AllocDescriptor, (uint64)(SizeInWholePages + DescriptorSize));
+			// unreachable
+			return;
+		}
 
-			if (UNLIKELY(AllocDescriptor->OriginalSizeAsPassed != Size))
-			{
-				UE_LOG(LogHAL, Fatal, TEXT("BinnedFreeToOS(): info mismatch: descriptor original size %llu, our size is %llu for pointer %p"), AllocDescriptor->OriginalSizeAsPassed, Size, Ptr);
-				// unreachable
-				return;
-			}
+		if (UNLIKELY(AllocDescriptor->OriginalSizeAsPassed != Size))
+		{
+			UE_LOG(LogHAL, Fatal, TEXT("BinnedFreeToOS(): info mismatch: descriptor original size %llu, our size is %llu for pointer %p"), AllocDescriptor->OriginalSizeAsPassed, Size, Ptr);
+			// unreachable
+			return;
 		}
 
 		AllocDescriptor = nullptr;	// just so no one touches it
@@ -592,7 +574,7 @@ void FUnixPlatformMemory::BinnedFreeToOS(void* Ptr, SIZE_T Size)
 			// unreachable
 		}
 	}
-	else
+	#else
 	{
 		UE::FForkPageProtector::Get().FreeMemoryRegion(Ptr);
 
@@ -602,6 +584,7 @@ void FUnixPlatformMemory::BinnedFreeToOS(void* Ptr, SIZE_T Size)
 			// unreachable
 		}
 	}
+	#endif
 }
 
 size_t FUnixPlatformMemory::FPlatformVirtualMemoryBlock::GetVirtualSizeAlignment()
