Fix virtual memory stats on Unix/Linux/Android

Typically this stat indicates RAM + swap.

[RN] Fixed Virtual Memory reporting on Linux/Unix/Android.

#Synced-CL 50439665

diff --git a/Engine/Source/Runtime/Core/Private/Android/AndroidPlatformMemory.cpp b/Engine/Source/Runtime/Core/Private/Android/AndroidPlatformMemory.cpp
--- a/Engine/Source/Runtime/Core/Private/Android/AndroidPlatformMemory.cpp
+++ b/Engine/Source/Runtime/Core/Private/Android/AndroidPlatformMemory.cpp
@@ -427,7 +427,9 @@ const FPlatformMemoryConstants& FAndroidPlatformMemory::GetConstants()
 		if (0 == sysinfo(&SysInfo))
 		{
 			MaxPhysicalRAMBytes = static_cast< unsigned long long >(SysInfo.mem_unit) * static_cast< unsigned long long >(SysInfo.totalram);
-			MaxVirtualRAMBytes = static_cast< unsigned long long >(SysInfo.mem_unit) * static_cast< unsigned long long >(SysInfo.totalswap);
+			MaxVirtualRAMBytes
+				= MaxPhysicalRAMBytes
+				+ static_cast< unsigned long long >(SysInfo.mem_unit) * static_cast< unsigned long long >(SysInfo.totalswap);
 		}
 
 		MemoryConstants.TotalPhysical = MaxPhysicalRAMBytes;
diff --git a/Engine/Source/Runtime/Core/Private/Unix/UnixPlatformMemory.cpp b/Engine/Source/Runtime/Core/Private/Unix/UnixPlatformMemory.cpp
--- a/Engine/Source/Runtime/Core/Private/Unix/UnixPlatformMemory.cpp
+++ b/Engine/Source/Runtime/Core/Private/Unix/UnixPlatformMemory.cpp
@@ -1216,7 +1216,9 @@ const FPlatformMemoryConstants& FUnixPlatformMemory::GetConstants()
 		if (0 == sysinfo(&SysInfo))
 		{
 			MaxPhysicalRAMBytes = static_cast< unsigned long long >( SysInfo.mem_unit ) * static_cast< unsigned long long >( SysInfo.totalram );
-			MaxVirtualRAMBytes = static_cast< unsigned long long >( SysInfo.mem_unit ) * static_cast< unsigned long long >( SysInfo.totalswap );
+			MaxVirtualRAMBytes
+				= MaxPhysicalRAMBytes
+				+ static_cast< unsigned long long >( SysInfo.mem_unit ) * static_cast< unsigned long long >( SysInfo.totalswap );
 		}
 
 		MemoryConstants.TotalPhysical = MaxPhysicalRAMBytes;
diff --git a/Engine/Source/Runtime/Core/Public/GenericPlatform/GenericPlatformMemory.h b/Engine/Source/Runtime/Core/Public/GenericPlatform/GenericPlatformMemory.h
--- a/Engine/Source/Runtime/Core/Public/GenericPlatform/GenericPlatformMemory.h
+++ b/Engine/Source/Runtime/Core/Public/GenericPlatform/GenericPlatformMemory.h
@@ -81,11 +81,11 @@ enum class EMemcpyCachePolicy : uint8
 
 /** 
  * Struct used to hold common memory constants for all platforms.
- * These values don't change over the entire life of the executable.
+ * These values generally don't change over the life of the executable, except for memory hotplugging and the addition or removal of swap space.
  */
 struct FGenericPlatformMemoryConstants
 {
-	/** The amount of actual physical memory, in bytes (needs to handle >4GB for 64-bit devices running 32-bit code). */
+	/** The amount of actual physical memory, in bytes. */
 	uint64 TotalPhysical = 0;
 
 	/** The amount of virtual memory, in bytes. */
