Add FFileHelper::LoadFileToStringAsync returning a TTask

#Synced-CL 50496977

diff --git a/Engine/Source/Runtime/Core/Private/Misc/FileHelper.cpp b/Engine/Source/Runtime/Core/Private/Misc/FileHelper.cpp
--- a/Engine/Source/Runtime/Core/Private/Misc/FileHelper.cpp
+++ b/Engine/Source/Runtime/Core/Private/Misc/FileHelper.cpp
@@ -1,11 +1,14 @@
 // Copyright Epic Games, Inc. All Rights Reserved.
 
 #include "Misc/FileHelper.h"
+
+#include "Async/AsyncFileHandle.h"
 #include "AutoRTFM.h"
 #include "Containers/StringConv.h"
 #include "Logging/LogMacros.h"
 #include "CoreGlobals.h"
 #include "HAL/PlatformMath.h"
+#include "Math/NumericLimits.h"
 #include "Memory/MemoryView.h"
 #include "Memory/SharedBuffer.h"
 #include "Misc/ByteSwap.h"
@@ -290,6 +293,105 @@ bool FFileHelper::LoadFileToString(FString& Result, IPlatformFile* PlatformFile,
 	return bSuccess;
 }
 
+UE::Tasks::TTask<TValueOrError<FString, FFileSystemError>> FFileHelper::LoadFileToStringAsync(const TCHAR* Filename, FLoadFileAsyncOptions Options)
+{
+	IPlatformFile& PlatformFile = Options.OverridePlatformFile ? *Options.OverridePlatformFile : FPlatformFileManager::Get().GetPlatformFile();
+	FFileOpenAsyncResult Result = PlatformFile.OpenAsyncRead(Filename, IPlatformFile::EOpenReadFlags::None);
+	if (Result.HasError())
+	{
+		return UE::Tasks::MakeCompletedTask<TValueOrError<FString, FFileSystemError>>(MakeError(Result.StealError()));
+	}
+
+	// Requests must be deleted before the handle, and cannot be deleted in their own callbacks
+	struct FState
+	{
+		TUniquePtr<IAsyncReadFileHandle> Handle;
+		// The previous request, either size or data, waiting to be deleted when we're no longer in the associated callback and it's safe to do so
+		TUniquePtr<IAsyncReadRequest> RequestToDelete;
+		UE::Tasks::FTaskEvent DataEvent{TEXT("LoadFileToStringAsync_DataEvent")};
+		int64 Size = 0;
+		TUniquePtr<uint8[]> Data;
+		FString HashFilename;
+		EAsyncIOPriorityAndFlags PriorityAndFlags;
+		EHashOptions VerifyFlags;
+	};
+	TSharedPtr<FState> State = MakeShared<FState>(FState
+	{
+		.Handle = Result.StealValue(),
+		.HashFilename = EnumHasAnyFlags(Options.VerifyFlags, EHashOptions::EnableVerify) ? FString(Filename) : FString(),
+		.PriorityAndFlags = Options.PriorityAndFlags,
+		.VerifyFlags = Options.VerifyFlags
+	});
+
+	// Note that callbacks are passed by ptr because they are optional.
+	// The API doesn't require the supplied pointer to live past the function call, just that the copy remain safely callable.
+	TFunction<void(bool, IAsyncReadRequest*)> SizeCallback = [State](bool bCancelled, IAsyncReadRequest* Request) mutable
+	{
+		State->RequestToDelete.Reset(Request);
+		State->Size = Request->GetSizeResults();
+		if (State->Size <= 0 || State->Size > (int64)MAX_int32) // Result is -1 on cancellation or error
+		{
+			// Trigger final task without reading
+			State->DataEvent.Trigger();
+			return;
+		}
+		TFunction<void(bool, IAsyncReadRequest*)> ReadCallback = [State](bool bCancelled, IAsyncReadRequest* Request) mutable
+		{
+			State->RequestToDelete.Reset(Request);
+			State->DataEvent.Trigger();
+		};
+		State->Data = MakeUnique<uint8[]>(State->Size);
+		State->Handle->ReadRequest(0, State->Size, State->PriorityAndFlags, &ReadCallback, State->Data.Get());
+	};
+	State->Handle->SizeRequest(&SizeCallback);
+	
+	if (EnumHasAnyFlags(Options.VerifyFlags, EHashOptions::EnableVerify))
+	{
+		State->HashFilename = Filename;
+	}
+
+	return UE::Tasks::Launch(TEXT("LoadFileToStringAsync"), [State]() -> TValueOrError<FString, FFileSystemError>
+		{
+			if (State->Size < 0)
+			{
+				// Async read handle doesn't provide detailed errors at this point
+				return MakeError(FFileSystemError(TEXT("Unknown error")));
+			}
+			else if (State->Size == 0)
+			{
+				return MakeValue(FString());
+			}
+			else if (State->Size > (int64)MAX_int32)
+			{
+				return MakeError(FFileSystemError(FString::Printf(TEXT("File too large (%lld) for FString"), State->Size)));
+			}
+
+			FString String;
+			if (State->Data.IsValid())
+			{
+				FFileHelper::BufferToString(String, State->Data.Get(), (int32)State->Size);
+			}
+
+			// Handle SHA verification of the file.
+			if (EnumHasAnyFlags(State->VerifyFlags, EHashOptions::EnableVerify) 
+			&& (EnumHasAnyFlags(State->VerifyFlags, EHashOptions::ErrorMissingHash) || FSHA1::GetFileSHAHash(*State->HashFilename, nullptr)))
+			{
+				// Buffer will be freed after async verification is finished
+				FBufferReaderWithSHA Ar(
+					(void*)State->Data.Release(),
+					State->Size, 
+					/* free buffer on close */ true, 
+					*State->HashFilename, 
+					/* bIsPersistent */ false, 
+					/* bIsUnfoundHashAnError */ true);
+			}
+
+			return MakeValue(MoveTemp(String));
+		},
+		UE::Tasks::Prerequisites(State->DataEvent)
+	);
+}
+
 bool FFileHelper::LoadFileToStringArray( TArray<FString>& Result, const TCHAR* Filename )
 {
 	return LoadFileToStringArrayWithPredicate(Result, Filename, [](const FString&) { return true;  });
diff --git a/Engine/Source/Runtime/Core/Public/Misc/FileHelper.h b/Engine/Source/Runtime/Core/Public/Misc/FileHelper.h
--- a/Engine/Source/Runtime/Core/Public/Misc/FileHelper.h
+++ b/Engine/Source/Runtime/Core/Public/Misc/FileHelper.h
@@ -15,6 +15,7 @@
 #include "Math/MathFwd.h"
 #include "Memory/MemoryFwd.h"
 #include "Misc/EnumClassFlags.h"
+#include "Tasks/Task.h"
 #include "Templates/UnrealTemplate.h"
 
 class FArchive;
@@ -54,6 +55,14 @@ struct FFileHelper
 		All
 	};
 
+	/** Options for methods that load files, for use with designated initializers */
+	struct FLoadFileAsyncOptions
+	{
+		IPlatformFile* OverridePlatformFile			= nullptr;
+		EAsyncIOPriorityAndFlags PriorityAndFlags 	= AIOP_Normal;
+		IPlatformFile::EOpenReadFlags FileOpenFlags = IPlatformFile::EOpenReadFlags::None;
+		EHashOptions VerifyFlags 					= EHashOptions::None;
+	};
 	/**
 	 * Load a text file to an FString.
 	 * Supports all combination of ANSI/Unicode files and platforms.
@@ -130,6 +139,21 @@ struct FFileHelper
 	 */
 	static CORE_API bool LoadFileToString(FString& Result, IPlatformFile* PlatformFile, const TCHAR* Filename, EHashOptions VerifyFlags = EHashOptions::None, uint32 ReadFlags = 0);
 
+	/**
+	 * Asynchronously load a text file to an FString. Supports all combination of ANSI/Unicode files and platforms.
+	 *
+	 * @param Filename     Name of the file to load
+	 * @param Options	   Additional options to control behavior, use designated initializers for customization.
+	 *
+	 * e.g. 
+	 * @code
+	 * FFileHelper::LoadFileToStringAsync(TEXT("foo.txt"), {});
+	 * FFileHelper::LoadFileToStringAsync(TEXT("foo.txt"),
+	 * 		{ .OverridePlatformFile = MyPlatformFile, .FileOpenFlags = IPlatformFile::EOpenReadFlags::NoBuffering });
+	 * @endcode  
+	 */
+	static CORE_API UE::Tasks::TTask<TValueOrError<FString, FFileSystemError>> LoadFileToStringAsync(const TCHAR* Filename, FLoadFileAsyncOptions Options);
+
 	/**
 	 * Load a text file to an array of strings. Supports all combination of ANSI/Unicode files and platforms.
 	 *
