StringView: Make UTF8TEXTVIEW constexpr which required C++20

#Synced-CL 50253651

diff --git a/Engine/Source/Runtime/Core/Public/Containers/StringView.h b/Engine/Source/Runtime/Core/Public/Containers/StringView.h
--- a/Engine/Source/Runtime/Core/Public/Containers/StringView.h
+++ b/Engine/Source/Runtime/Core/Public/Containers/StringView.h
@@ -602,11 +602,9 @@ template <UE::CCharType CharType>
 	return FWideStringView(String, (int32)Size);
 }
 
-[[nodiscard]] /*constexpr*/ inline FUtf8StringView operator""_PrivateU8SV(const ANSICHAR* String, size_t Size)
+[[nodiscard]] constexpr inline FUtf8StringView operator""_PrivateU8SV(const UTF8CHAR* String, size_t Size)
 {
-	// Would like this operator to be constexpr, but cannot be until after this operator can take a UTF8CHAR*
-	// rather than an ANSICHAR*, which won't be until we have C++20 char8_t string literals.
-	return FUtf8StringView(reinterpret_cast<const UTF8CHAR*>(String), (int32)Size);
+	return FUtf8StringView(String, (int32)Size);
 }
 
 #if PLATFORM_TCHAR_IS_UTF8CHAR
@@ -616,7 +614,7 @@ template <UE::CCharType CharType>
 #endif
 #define ANSITEXTVIEW(str) (str##_PrivateASV)
 #define WIDETEXTVIEW(str) UE_JOIN(WIDETEXT(str), _PrivateWSV)
-#define UTF8TEXTVIEW(str) (str##_PrivateU8SV)
+#define UTF8TEXTVIEW(str) UE_JOIN(UTF8TEXT(str), _PrivateU8SV)
 
 namespace UE::Core::Private
 {
diff --git a/Engine/Source/Runtime/Core/Tests/Misc/AsciiSetTest.cpp b/Engine/Source/Runtime/Core/Tests/Misc/AsciiSetTest.cpp
--- a/Engine/Source/Runtime/Core/Tests/Misc/AsciiSetTest.cpp
+++ b/Engine/Source/Runtime/Core/Tests/Misc/AsciiSetTest.cpp
@@ -37,7 +37,7 @@ TEST_CASE_NAMED(TAsciiSetTest, "System::Core::Misc::AsciiSet", "[ApplicationCont
 
 	CHECK_EQUALS(TEXT("Skip"), FAsciiSet::Skip(TEXT("  \t\tHello world!"), Whitespaces), TEXT("Hello world!"));
 	CHECK_EQUALS(TEXT("Skip"), FAsciiSet::Skip(TEXT("Hello world!"), Whitespaces), TEXT("Hello world!"));
-	CHECK_EQUALS(TEXT("Skip to extended ASCII"), FAsciiSet::Skip(" " "\xA0" " abc", Whitespaces), UTF8TEXTVIEW("\xA0" " abc"));
+	CHECK_EQUALS(TEXT("Skip to extended ASCII"), FAsciiSet::Skip(" " "\xA0" " abc", Whitespaces), ANSITEXTVIEW("\xA0" " abc"));
 	CHECK_EQUALS(TEXT("Skip to wide"), FAsciiSet::Skip(TEXT(" 变 abc"), Whitespaces), TEXT("变 abc"));
 	CHECK_EQUALS(TEXT("AdvanceToFirst"),	*FAsciiSet::FindFirstOrEnd("NonWhitespace\t \nNonWhitespace", Whitespaces), '\t');
 	CHECK_EQUALS(TEXT("AdvanceToLast"),	*FAsciiSet::FindLastOrEnd("NonWhitespace\t \nNonWhitespace", Whitespaces), '\n');
diff --git a/Engine/Source/Runtime/Core/Tests/Serialization/CompactBinaryWriterTest.cpp b/Engine/Source/Runtime/Core/Tests/Serialization/CompactBinaryWriterTest.cpp
--- a/Engine/Source/Runtime/Core/Tests/Serialization/CompactBinaryWriterTest.cpp
+++ b/Engine/Source/Runtime/Core/Tests/Serialization/CompactBinaryWriterTest.cpp
@@ -346,7 +346,7 @@ TEST_CASE_NAMED(FCbWriterStringTest, "System::Core::Serialization::CbWriter::Str
 		{
 			for (FCbFieldView Field : Fields)
 			{
-				CHECK_EQUALS(TEXT("FCbWriter(String, Unicode).AsString()"), Field.AsString(), UTF8TEXTVIEW("\xf0\x9f\x98\x80"));
+				CHECK_EQUALS(TEXT("FCbWriter(String, Unicode).AsString()"), Field.AsString(), ANSITEXTVIEW("\xf0\x9f\x98\x80"));
 			}
 		}
 	}
