Sanitize test name at registration

#rnx

#Synced-CL 50593455

diff --git a/Engine/Source/Runtime/Core/Private/Misc/AutomationTest.cpp b/Engine/Source/Runtime/Core/Private/Misc/AutomationTest.cpp
--- a/Engine/Source/Runtime/Core/Private/Misc/AutomationTest.cpp
+++ b/Engine/Source/Runtime/Core/Private/Misc/AutomationTest.cpp
@@ -149,18 +149,23 @@ namespace AutomationTest
 
 		return EffectiveVerbosity;
 	}
+
+	static inline bool GetTestClassAndParameters(const FString& InTest, FString* OutTestClass, FString* OutTestParameters)
+	{
+		return InTest.Split(TEXT(" "), OutTestClass, OutTestParameters, ESearchCase::CaseSensitive);
+	}
 };
 
 FAutomationTestBase::FAutomationTestBase( const FString& InName, const bool bInComplexTask )
 	: bComplexTask( bInComplexTask )
 {
 	LLM_SCOPE_BYNAME(TEXT("AutomationTest/Framework"));
-	TestName = InName;
+	TestName = InName.Replace(TEXT(" "), TEXT("_"));
 	// Register the newly created automation test into the automation testing framework
-	const bool bRegistered = FAutomationTestFramework::Get().RegisterAutomationTest( InName, this );
+	const bool bRegistered = FAutomationTestFramework::Get().RegisterAutomationTest( TestName, this );
 	if (!bRegistered)
 	{
-		UE_LOG(LogAutomationTest, Warning, TEXT("Failed to register test with the name '%s'. Test with the same name is already registered and will not be overridden."), *InName);
+		UE_LOG(LogAutomationTest, Warning, TEXT("Failed to register test with the name '%s'. Test with the same name is already registered and will not be overridden."), *TestName);
 	}
 }
 
@@ -628,7 +633,7 @@ void FAutomationTestFramework::StartTestByName( const FString& InTestToRun, cons
 
 	FString TestName;
 	FString Params;
-	if (!InTestToRun.Split(TEXT(" "), &TestName, &Params, ESearchCase::CaseSensitive))
+	if (!AutomationTest::GetTestClassAndParameters(InTestToRun, &TestName, &Params))
 	{
 		TestName = InTestToRun;
 	}
@@ -1270,7 +1275,7 @@ void FAutomationTestFramework::InternalStartTest( const FString& InTestToRun, co
 	CurrentTestFullPath.Empty();
 
 	FString TestName;
-	if (!InTestToRun.Split(TEXT(" "), &TestName, &Parameters, ESearchCase::CaseSensitive))
+	if (!AutomationTest::GetTestClassAndParameters(InTestToRun, &TestName, &Parameters))
 	{
 		TestName = InTestToRun;
 	}
@@ -1412,13 +1417,23 @@ bool FAutomationTestFramework::CanRunTestInEnvironment(const FString& InTestToRu
 {
 	FString TestClassName;
 	FString TestParameters;
-	if (!InTestToRun.Split(TEXT(" "), &TestClassName, &TestParameters, ESearchCase::CaseSensitive))
+	if (!AutomationTest::GetTestClassAndParameters(InTestToRun, &TestClassName, &TestParameters))
 	{
 		TestClassName = InTestToRun;
 	}
 
 	if (!ContainsTest(TestClassName))
 	{
+		if (OutReason != nullptr)
+		{
+			*OutReason = FString::Printf(TEXT("Could not find test %s in test registry."), *TestClassName);
+		}
+
+		if (OutWarn != nullptr)
+		{
+			*OutWarn = true;
+		}
+
 		return false;
 	}
 
