Added the required IntrusiveUnsetOptionalStateType alias to the example of how to defined a class with an intrusive unset optional state.

diff --git a/Engine/Source/Runtime/Core/Public/Misc/IntrusiveUnsetOptionalState.h b/Engine/Source/Runtime/Core/Public/Misc/IntrusiveUnsetOptionalState.h
--- a/Engine/Source/Runtime/Core/Public/Misc/IntrusiveUnsetOptionalState.h
+++ b/Engine/Source/Runtime/Core/Public/Misc/IntrusiveUnsetOptionalState.h
@@ -27,6 +27,9 @@
  *     // This static member should be constexpr, public and equal to true.
  *     static constexpr bool bHasIntrusiveUnsetOptionalState = true;
  *
+ *     // This alias should be set to the class type itself.
+ *     using IntrusiveUnsetOptionalStateType = FMyType;
+ *
  *     explicit FMyType(int32 InIndex)
  *     {
  *         // Validate class invariant.
