Fix broken Clang 128bit integer math in GPU timestamp transposition to CPU time domain.

#rnx

#Synced-CL 50626303

diff --git a/Engine/Source/Runtime/Core/Public/Math/Int128.h b/Engine/Source/Runtime/Core/Public/Math/Int128.h
--- a/Engine/Source/Runtime/Core/Public/Math/Int128.h
+++ b/Engine/Source/Runtime/Core/Public/Math/Int128.h
@@ -279,4 +279,13 @@ public:
 	{
 		return Ar << Value.Hi << Value.Lo;
 	}
+
+	/**
+	 * Casting
+	 */
+	inline operator uint64() const
+	{
+		check(Hi == 0);
+		return Lo;
+	}
 };
\ No newline at end of file
