Add ThreadManager OnAddThread/OnRemoveThread hooks

#rnx

0) Integration of https://p4-swarm.epicgames.net/changes/18474135, https://p4-swarm.epicgames.net/changes/18474917
0.1) It allows to track all UE threads for external thread management systems.
1) Added FThreadManager::EThreadEventType enum to identify thread event types (e.g. Add or Remove).
2) Added thread-safe ThreadEvent delegate for ThreadManager. It is called in methods AddThread and RemoveThread of ThreadManager.
For more context:
- pthread RunnableThread clears ThreadID, and it can be empty => that instance cannot be used to take ThreadId. ThreadManager knows the correct id extended to a related event to pass that id.
- RemoveThreadId* is invalidated after Remove.
2) The delegates are raised once for each RunnableThread instance.

Test plan:
1) Verify compilation by running rtc tests. https://horde.devtools.epicgames.com/job/693c1ea55d767020c6a9df02

#Synced-CL 50588029

diff --git a/Engine/Source/Runtime/Core/Private/HAL/ThreadingBase.cpp b/Engine/Source/Runtime/Core/Private/HAL/ThreadingBase.cpp
--- a/Engine/Source/Runtime/Core/Private/HAL/ThreadingBase.cpp
+++ b/Engine/Source/Runtime/Core/Private/HAL/ThreadingBase.cpp
@@ -9,6 +9,7 @@
 #include "Misc/EventPool.h"
 #include "Misc/LazySingleton.h"
 #include "Misc/Fork.h"
+#include "Misc/Optional.h"
 #include "Templates/Atomic.h"
 #include "HAL/IConsoleManager.h"
 #include "HAL/PlatformStackWalk.h"
@@ -520,27 +521,50 @@ void FThreadManager::AddThread(uint32 ThreadId, FRunnableThread* Thread)
 		checkf((ThreadId & FFakeThread::FakeIdReservedBit) == 0, TEXT("The thread ID  assigned by the kernel clashes with the bit reserved for identifying fake threads. Need to revisit the fake ID assignment algo."));
 	}
 
-	FScopeLock ThreadsLock(&ThreadsCritical);	
-
-	// Some platforms do not support TLS
-	if (!Threads.Contains(ThreadId))
+	bool bAlreadyInThreadPool = false;
 	{
-		Threads.Add(ThreadId, Thread);
-		OnThreadListModified();
+		const FScopeLock ThreadsLock(&ThreadsCritical);
+		// Some platforms do not support TLS
+		bAlreadyInThreadPool = Threads.Contains(ThreadId);
+		if (!bAlreadyInThreadPool)
+		{
+			Threads.Add(ThreadId, Thread);
+			OnThreadListModified();
+		}
+	}
+	// notify observers only about new threads
+	if (!bAlreadyInThreadPool)
+	{
+		ThreadEvent.Broadcast(ThreadId, Thread, EThreadEventType::Added);
 	}
 }
 
 void FThreadManager::RemoveThread(FRunnableThread* Thread)
 {
-	FScopeLock ThreadsLock(&ThreadsCritical);
-	const uint32* ThreadId = Threads.FindKey(Thread);
-	if (ThreadId)
+	TOptional<uint32> RemovedThreadId;
 	{
-		Threads.Remove(*ThreadId);
-		OnThreadListModified();
+		FScopeLock ThreadsLock(&ThreadsCritical);
+		const uint32* ThreadId = Threads.FindKey(Thread);
+		if (ThreadId)
+		{
+			RemovedThreadId.Emplace(*ThreadId);
+			// `ThreadId` can be invalidated after it
+			Threads.Remove(*ThreadId);
+			OnThreadListModified();
+		}
+	}
+	// notify observers only about existing threads
+	if (RemovedThreadId)
+	{
+		ThreadEvent.Broadcast(RemovedThreadId.GetValue(), Thread, EThreadEventType::Removed);
 	}
 }
 
+FThreadManager::FThreadEvent& FThreadManager::OnThreadEvent()
+{
+	return ThreadEvent;
+}
+
 void FThreadManager::Tick()
 {	
 	const bool bIsSingleThreadEnvironment = FPlatformProcess::SupportsMultithreading() == false;
diff --git a/Engine/Source/Runtime/Core/Public/HAL/ThreadManager.h b/Engine/Source/Runtime/Core/Public/HAL/ThreadManager.h
--- a/Engine/Source/Runtime/Core/Public/HAL/ThreadManager.h
+++ b/Engine/Source/Runtime/Core/Public/HAL/ThreadManager.h
@@ -5,6 +5,7 @@
 #include "CoreTypes.h"
 #include "CoreGlobals.h"
 #include "Containers/Map.h"
+#include "Delegates/Delegate.h"
 #include "HAL/CriticalSection.h"
 
 #ifndef PLATFORM_SUPPORTS_ALL_THREAD_BACKTRACES
@@ -52,6 +53,25 @@ public:
 	*/
 	CORE_API void RemoveThread(FRunnableThread* Thread);
 
+
+	/** Specifies the type of thread event. */
+	enum class EThreadEventType
+	{
+		Added,
+		Removed
+	};
+
+	/**
+	  * Used internally to broadcast about a real thread object events.
+	  * A thread-safe version, broadcast might be done from any thread.
+	  * 
+	  * @see EThreadEventType
+	  * @see AddThread
+	  * @see RemoveThread
+	  */
+	DECLARE_TS_MULTICAST_DELEGATE_ThreeParams(FThreadEvent, const uint32 /*ThreadId*/, class FRunnableThread* /*Thread*/, EThreadEventType /*EventType*/);
+	FThreadEvent& OnThreadEvent();
+
 	/** Get the number of registered threads */
 	int32 NumThreads() const { return Threads.Num(); }
 
@@ -129,4 +149,7 @@ private:
 
 	/** Notification that the parent is about to fork */
 	void HandleOnParentPreFork();
+
+	/** Broadcasts whenever a thread has been added or removed. See EThreadEventType for details */
+	FThreadEvent ThreadEvent;
 };
