Getting an accurate ProgramSize stat on iOS via dladdr

#Synced-CL 50657985

diff --git a/Engine/Source/Runtime/Core/Private/HAL/LLM/LLM.cpp b/Engine/Source/Runtime/Core/Private/HAL/LLM/LLM.cpp
--- a/Engine/Source/Runtime/Core/Private/HAL/LLM/LLM.cpp
+++ b/Engine/Source/Runtime/Core/Private/HAL/LLM/LLM.cpp
@@ -2311,7 +2311,8 @@ void FLLMGlobals::PublishDataPerFrame(const TCHAR* LogName)
 void FLLMGlobals::InitializeProgramSize()
 {
 	// On Android we report ProgramSize via AndroidMemoryProbeTick
-#if !PLATFORM_ANDROID
+	// and on iOS we report ProgramSize via FAppleMemoryProbe::ScanMemory
+#if !PLATFORM_ANDROID && !PLATFORM_IOS
 	if (!ProgramSize)
 	{
 		FPlatformMemoryStats Stats = FPlatformMemory::GetStats();
@@ -7600,4 +7601,4 @@ FORCENOINLINE void* FArraySlackTrackingHeader::Realloc(void* Ptr, int64 Count, u
 }
 #endif
 
-#endif  // #else .. #if ENABLE_LOW_LEVEL_MEM_TRACKER
\ No newline at end of file
+#endif  // #else .. #if ENABLE_LOW_LEVEL_MEM_TRACKER
diff --git a/Engine/Source/Runtime/Core/Private/ProfilingDebugging/Apple/AppleMemoryProbe.cpp b/Engine/Source/Runtime/Core/Private/ProfilingDebugging/Apple/AppleMemoryProbe.cpp
--- a/Engine/Source/Runtime/Core/Private/ProfilingDebugging/Apple/AppleMemoryProbe.cpp
+++ b/Engine/Source/Runtime/Core/Private/ProfilingDebugging/Apple/AppleMemoryProbe.cpp
@@ -19,6 +19,8 @@
 #include <sys/types.h>
 #include <sys/mman.h>
 
+#include <dlfcn.h>
+
 extern bool MemoryTrace_IsActive();
 
 static float GAppleMemoryProbeInterval = 5.0f;
@@ -148,6 +150,7 @@ private:
 		natural_t Depth = 0;
 		vm_region_submap_info_data_64_t Info = {};
 		mach_msg_type_number_t InfoCount = VM_REGION_SUBMAP_INFO_COUNT_64;
+		uint64_t ProgramSize = 0;
 
 		// Go over our memory map and gather all memory regions to track
 		while (true)
@@ -161,8 +164,18 @@ private:
 			const vm_address_t CurrAddr = NextAddr;
 			NextAddr += Size;
 			
-			// TODO future work: collect info about resident pages on program images and report them to LLM
-
+			if(!Info.is_submap)
+			{
+				Dl_info dl_info{};
+				dladdr(reinterpret_cast<void*>(CurrAddr), &dl_info);
+				
+				if(dl_info.dli_fname != nullptr && 
+				   FAnsiStringView(dl_info.dli_fname).EndsWith(TCHAR_TO_ANSI(FPlatformProcess::ExecutableName())))
+				{
+					ProgramSize += Info.pages_resident * vm_page_size;
+				}
+			}
+			
 			if (!ShouldScanRange(Info.user_tag))
 			{
 				continue;
@@ -186,6 +199,11 @@ private:
 			}
 		}
 		
+	#if ENABLE_LOW_LEVEL_MEM_TRACKER
+		FLowLevelMemTracker::Get().SetProgramSize(ProgramSize);
+	#endif
+		FGenericPlatformMemory::SetProgramSize(ProgramSize);
+		
 		struct
 		{
 			uint32 PageMaps;
