Remove constexpr from FLogCategoryBase functions that have non-const members in them
#Synced-CL 50661700

diff --git a/Engine/Source/Runtime/Core/Public/Logging/LogCategory.h b/Engine/Source/Runtime/Core/Public/Logging/LogCategory.h
--- a/Engine/Source/Runtime/Core/Public/Logging/LogCategory.h
+++ b/Engine/Source/Runtime/Core/Public/Logging/LogCategory.h
@@ -42,7 +42,7 @@ struct FLogCategoryBase
 	UE_API ~FLogCategoryBase();
 
 	/** Should not generally be used directly. Tests the runtime verbosity and maybe triggers a debug break, etc. */
-	UE_FORCEINLINE_HINT constexpr bool IsSuppressed(ELogVerbosity::Type VerbosityLevel) const
+	UE_FORCEINLINE_HINT bool IsSuppressed(ELogVerbosity::Type VerbosityLevel) const
 	{
 		return !((VerbosityLevel & ELogVerbosity::VerbosityMask) <= Verbosity);
 	}
@@ -50,16 +50,16 @@ struct FLogCategoryBase
 	/** Called just after a logging statement being allow to print. Checks a few things and maybe breaks into the debugger. */
 	UE_API void PostTrigger(ELogVerbosity::Type VerbosityLevel);
 
-	inline constexpr const FLogCategoryName& GetCategoryName() const { return CategoryName; }
+	inline const FLogCategoryName& GetCategoryName() const { return CategoryName; }
 
 	/** Gets the working verbosity. */
-	inline constexpr ELogVerbosity::Type GetVerbosity() const { return (ELogVerbosity::Type)Verbosity; }
+	inline ELogVerbosity::Type GetVerbosity() const { return (ELogVerbosity::Type)Verbosity; }
 	
 	/** Sets up the working verbosity and clamps to the compile time verbosity. */
 	UE_API void SetVerbosity(ELogVerbosity::Type Verbosity);
 
 	/** Gets the compile time verbosity. */
-	inline constexpr ELogVerbosity::Type GetCompileTimeVerbosity() const { return CompileTimeVerbosity; }
+	inline ELogVerbosity::Type GetCompileTimeVerbosity() const { return CompileTimeVerbosity; }
 
 private:
 	friend class FLogSuppressionImplementation;
