Fixing nonunity build breakage with fixing ProgramSize memory stat on iOS

#Synced-CL 50664794

diff --git a/Engine/Source/Runtime/Core/Private/ProfilingDebugging/Apple/AppleMemoryProbe.cpp b/Engine/Source/Runtime/Core/Private/ProfilingDebugging/Apple/AppleMemoryProbe.cpp
--- a/Engine/Source/Runtime/Core/Private/ProfilingDebugging/Apple/AppleMemoryProbe.cpp
+++ b/Engine/Source/Runtime/Core/Private/ProfilingDebugging/Apple/AppleMemoryProbe.cpp
@@ -5,6 +5,7 @@
 #include "HAL/Thread.h"
 #include "HAL/PlatformProcess.h"
 #include "HAL/PlatformAffinity.h"
+#include "HAL/LowLevelMemTracker.h"
 
 #if UE_MEMORY_TRACE_ENABLED
 
