Trace scopes

#Synced-CL 50645574

diff --git a/Engine/Source/Runtime/Core/Private/Internationalization/LocalizationResourceTextSource.cpp b/Engine/Source/Runtime/Core/Private/Internationalization/LocalizationResourceTextSource.cpp
--- a/Engine/Source/Runtime/Core/Private/Internationalization/LocalizationResourceTextSource.cpp
+++ b/Engine/Source/Runtime/Core/Private/Internationalization/LocalizationResourceTextSource.cpp
@@ -9,6 +9,8 @@
 #include "Misc/CoreDelegates.h"
 #include "Misc/EnumClassFlags.h"
 #include "Misc/Paths.h"
+#include "Trace/Trace.h"
+#include "Trace/Trace.inl"
 
 bool FLocalizationResourceTextSource::GetNativeCultureName(const ELocalizedTextSourceCategory InCategory, FString& OutNativeCultureName)
 {
@@ -149,6 +151,11 @@ void FLocalizationResourceTextSource::LoadLocalizedResources(const ELocalization
 	LoadLocalizedResourcesFromPaths(PrioritizedNativePaths, PrioritizedLocalizationPaths, GameNativePaths, InLoadFlags, InPrioritizedCultures, InOutNativeResource, InOutLocalizedResource);
 }
 
+UE_TRACE_EVENT_BEGIN(Cpu, LoadLocalizationResourcesForCulture, NoSync)
+UE_TRACE_EVENT_FIELD(UE::Trace::WideString, Culture)
+UE_TRACE_EVENT_FIELD(UE::Trace::WideString, Filename)
+UE_TRACE_EVENT_END()
+
 void FLocalizationResourceTextSource::LoadLocalizedResourcesFromPaths(TArrayView<const FString> InPrioritizedNativePaths, TArrayView<const FString> InPrioritizedLocalizationPaths, TArrayView<const FString> InGameNativePaths, const ELocalizationLoadFlags InLoadFlags, TArrayView<const FString> InPrioritizedCultures, FTextLocalizationResource& InOutNativeResource, FTextLocalizationResource& InOutLocalizedResource) const
 {
 	SCOPED_BOOT_TIMING("LoadLocalizedResourcesFromPaths");
@@ -170,12 +177,24 @@ void FLocalizationResourceTextSource::LoadLocalizedResourcesFromPaths(TArrayView
 			}
 		}
 
-		const FString PlatformAgnosticLocResFilename = InLocalizationPath / InCulture / InLocResFilename;
-		InOutLocRes.LoadFromFile(PlatformAgnosticLocResFilename, InLocResPriority);
+		{ 
+			const FString PlatformAgnosticLocResFilename = InLocalizationPath / InCulture / InLocResFilename;
+#if CPUPROFILERTRACE_ENABLED
+			UE_TRACE_LOG_SCOPED_T(Cpu, LoadLocalizationResourcesForCulture, CpuChannel)
+				<< LoadLocalizationResourcesForCulture.Culture(*InCulture)
+				<< LoadLocalizationResourcesForCulture.Filename(*PlatformAgnosticLocResFilename);
+#endif // CPUPROFILERTRACE_ENABLED
+			InOutLocRes.LoadFromFile(PlatformAgnosticLocResFilename, InLocResPriority);
+		}
 
 		const FString PlatformSpecificLocResFilename = InLocalizationPath / PlatformLocalizationFolderName / PlatformName / InCulture / InLocResFilename;
 		if (FPaths::FileExists(PlatformSpecificLocResFilename))
 		{
+#if CPUPROFILERTRACE_ENABLED
+			UE_TRACE_LOG_SCOPED_T(Cpu, LoadLocalizationResourcesForCulture, CpuChannel)
+				<< LoadLocalizationResourcesForCulture.Culture(*InCulture)
+				<< LoadLocalizationResourcesForCulture.Filename(*PlatformSpecificLocResFilename);
+#endif // CPUPROFILERTRACE_ENABLED
 			InOutLocRes.LoadFromFile(PlatformSpecificLocResFilename, InLocResPriority);
 		}
 	};
diff --git a/Engine/Source/Runtime/Core/Private/Internationalization/TextLocalizationResource.cpp b/Engine/Source/Runtime/Core/Private/Internationalization/TextLocalizationResource.cpp
--- a/Engine/Source/Runtime/Core/Private/Internationalization/TextLocalizationResource.cpp
+++ b/Engine/Source/Runtime/Core/Private/Internationalization/TextLocalizationResource.cpp
@@ -15,6 +15,8 @@
 #include "Misc/FileHelper.h"
 #include "Misc/CommandLine.h"
 #include "Serialization/MemoryReader.h"
+#include "Trace/Trace.h"
+#include "Trace/Trace.inl"
 
 DEFINE_LOG_CATEGORY_STATIC(LogTextLocalizationResource, Log, All);
 
@@ -47,8 +49,17 @@ int32 GetLocalizationTargetPathIdFromLocResId(const FTextKey& InLocResID)
 
 }
 
+UE_TRACE_EVENT_BEGIN(Cpu, LoadTextLocalizationMetaDataResourceFromFile, NoSync)
+UE_TRACE_EVENT_FIELD(UE::Trace::WideString, Path)
+UE_TRACE_EVENT_END()
+
 bool FTextLocalizationMetaDataResource::LoadFromFile(const FString& FilePath)
 {
+#if CPUPROFILERTRACE_ENABLED
+	UE_TRACE_LOG_SCOPED_T(Cpu, LoadTextLocalizationMetaDataResourceFromFile, CpuChannel)
+		<< LoadTextLocalizationMetaDataResourceFromFile.Path(*FilePath);
+#endif // CPUPROFILERTRACE_ENABLED
+
 	TUniquePtr<FArchive> Reader;
 #if PRELOAD_LOCMETA_FILES
 	TArray<uint8> FileBytes;
diff --git a/Engine/Source/Runtime/Core/Private/Misc/ConfigContext.cpp b/Engine/Source/Runtime/Core/Private/Misc/ConfigContext.cpp
--- a/Engine/Source/Runtime/Core/Private/Misc/ConfigContext.cpp
+++ b/Engine/Source/Runtime/Core/Private/Misc/ConfigContext.cpp
@@ -18,6 +18,8 @@
 #include "Logging/MessageLog.h"
 #include "Async/Async.h"
 #include "Async/TaskGraphInterfaces.h"
+#include "Trace/Trace.h"
+#include "Trace/Trace.inl"
 
 extern double GPrepareForLoadTime;
 extern double GPerformLoadTime;
@@ -224,9 +226,17 @@ static FConfigFile& ActiveFile(FConfigContext* Context)
 }
 
 
+UE_TRACE_EVENT_BEGIN(Cpu, FConfigContextLoad, NoSync)
+UE_TRACE_EVENT_FIELD(UE::Trace::WideString, BaseName)
+UE_TRACE_EVENT_END()
+
 bool FConfigContext::Load(const TCHAR* InBaseIniName, FString& OutFinalFilename)
 {
 	TRACE_CPUPROFILER_EVENT_SCOPE(FConfigContext::Load);
+#if CPUPROFILERTRACE_ENABLED
+	UE_TRACE_LOG_SCOPED_T(Cpu, FConfigContextLoad, CpuChannel)
+		<< FConfigContextLoad.BaseName(InBaseIniName);
+#endif // CPUPROFILERTRACE_ENABLED
 
 	if (Branch != nullptr && Branch->bIsSafeUnloaded)
 	{
diff --git a/Engine/Source/Runtime/Core/Private/UObject/UnrealNames.cpp b/Engine/Source/Runtime/Core/Private/UObject/UnrealNames.cpp
--- a/Engine/Source/Runtime/Core/Private/UObject/UnrealNames.cpp
+++ b/Engine/Source/Runtime/Core/Private/UObject/UnrealNames.cpp
@@ -1,6 +1,7 @@
 // Copyright Epic Games, Inc. All Rights Reserved.
 
 #include "UObject/UnrealNames.h"
+#include "ProfilingDebugging/CpuProfilerTrace.h"
 #include "UObject/NameBatchSerialization.h"
 #include "Async/Async.h"
 #include "Misc/AssertionMacros.h"
@@ -23,6 +24,7 @@
 #include "UObject/ObjectVersion.h"
 #include "HAL/ThreadSafeCounter.h"
 #include "Misc/ScopeRWLock.h"
+#include "ProfilingDebugging/CpuProfilerTrace.h"
 #include "ProfilingDebugging/StringsTrace.h"
 #include "ProfilingDebugging/MetadataTrace.h"
 #include "Containers/Set.h"
@@ -5275,6 +5277,7 @@ struct FNameBatchLoader
 	// @return true if there's anything to load
 	bool Read(FArchive& Ar)
 	{
+		TRACE_CPUPROFILER_EVENT_SCOPE_STR("FNameBatchLoader::Read");
 		uint32 Num = 0;
 		Ar << Num;
 
@@ -5306,6 +5309,7 @@ struct FNameBatchLoader
 
 	TArray<FDisplayNameEntryId> Load()
 	{
+		TRACE_CPUPROFILER_EVENT_SCOPE_STR("FNameBatchLoader::Load");
 		check(Headers.Num());
 
 		TArray<FDisplayNameEntryId> Out;
@@ -5365,6 +5369,7 @@ struct FNameBatchAsyncLoader : FNameBatchLoader
 
 	void DoWork()
 	{
+		TRACE_CPUPROFILER_EVENT_SCOPE_STR("FNameBatchAsyncLoader::DoWork");
 		if (EnumHasAnyFlags(Flags, ENameBatchLoadingFlags::RespectOrder))
 		{
 			LoadSeparatedNameBatchInInputOrder(Out, Hashes.GetData(), Headers, Strings);
@@ -5378,6 +5383,7 @@ struct FNameBatchAsyncLoader : FNameBatchLoader
 
 	TArray<FDisplayNameEntryId> GetResult()
 	{
+		TRACE_CPUPROFILER_EVENT_SCOPE_STR("FNameBatchAsyncLoader::GetResult");
 		check(DoneEvent);
 
 		DoneEvent->Wait();
diff --git a/Engine/Source/Runtime/CoreUObject/Private/UObject/CoreRedirects.cpp b/Engine/Source/Runtime/CoreUObject/Private/UObject/CoreRedirects.cpp
--- a/Engine/Source/Runtime/CoreUObject/Private/UObject/CoreRedirects.cpp
+++ b/Engine/Source/Runtime/CoreUObject/Private/UObject/CoreRedirects.cpp
@@ -2,6 +2,7 @@
 
 #include "UObject/CoreRedirects.h"
 
+#include "ProfilingDebugging/CpuProfilerTrace.h"
 #include "UObject/CoreRedirects/CoreRedirectsContext.h"
 #include "UObject/CoreRedirects/RedirectionSummary.h"
 #include "UObject/Linker.h"
@@ -2713,6 +2714,7 @@ bool FCoreRedirectTest::RunTest(const FString& Parameters)
 
 bool FCoreRedirects::ReadRedirectsFromIni(const FString& IniName)
 {
+	TRACE_CPUPROFILER_EVENT_SCOPE_STR("FCoreRedirects::ReadRedirectsFromIni");
 	FScopeCoreRedirectsWriteLockedContext LockedContext(FCoreRedirectsContext::GetThreadContext());
 	ensureMsgf(IsInitialized(), TEXT("FCoreRedirects must be initialized on the game thread before use."));
 
diff --git a/Engine/Source/Runtime/CrashReportCore/Private/CrashReportAnalytics.cpp b/Engine/Source/Runtime/CrashReportCore/Private/CrashReportAnalytics.cpp
--- a/Engine/Source/Runtime/CrashReportCore/Private/CrashReportAnalytics.cpp
+++ b/Engine/Source/Runtime/CrashReportCore/Private/CrashReportAnalytics.cpp
@@ -51,6 +51,7 @@ IAnalyticsProviderET& FCrashReportAnalytics::GetProvider()
 
 void FCrashReportAnalytics::Initialize(const FString& EpicAccountId)
 {
+	TRACE_CPUPROFILER_EVENT_SCOPE_STR("FCrashReportAnalytics::Initialize");
 	checkf(!bIsInitialized, TEXT("FCrashReportAnalytics::Initialize called more than once."));
 
 	// Allow build machines to force CRC to enable internal telemetry.
diff --git a/Engine/Source/Runtime/Json/Private/JsonUtils/RapidJsonUtils.cpp b/Engine/Source/Runtime/Json/Private/JsonUtils/RapidJsonUtils.cpp
--- a/Engine/Source/Runtime/Json/Private/JsonUtils/RapidJsonUtils.cpp
+++ b/Engine/Source/Runtime/Json/Private/JsonUtils/RapidJsonUtils.cpp
@@ -1,6 +1,7 @@
 // Copyright Epic Games, Inc. All Rights Reserved.
 
 #include "JsonUtils/RapidJsonUtils.h"
+#include "ProfilingDebugging/CpuProfilerTrace.h"
 
 namespace UE::Json
 {
@@ -232,6 +233,7 @@ FString FParseError::CreateMessage(const FStringView JsonText) const
 
 TValueOrError<FDocument, FParseError> Parse(const FStringView JsonText)
 {
+	TRACE_CPUPROFILER_EVENT_SCOPE_STR("UE::Json::Parse");
 	FDocument Result;
 	Result.Parse<DefaultParseFlags, FDocumentEncoding>(JsonText.GetData(), JsonText.Len());
 
@@ -248,6 +250,7 @@ TValueOrError<FDocument, FParseError> Parse(const FStringView JsonText)
 
 TValueOrError<FDocument, FParseError> ParseInPlace(TArrayView<TCHAR> JsonText)
 {
+	TRACE_CPUPROFILER_EVENT_SCOPE_STR("UE::Json::ParseInPlace");
 	if (JsonText.IsEmpty())
 	{
 		// the in-place read crashes if the text is empty
diff --git a/Engine/Source/Runtime/Projects/Private/PluginManifest.cpp b/Engine/Source/Runtime/Projects/Private/PluginManifest.cpp
--- a/Engine/Source/Runtime/Projects/Private/PluginManifest.cpp
+++ b/Engine/Source/Runtime/Projects/Private/PluginManifest.cpp
@@ -2,6 +2,7 @@
 
 #include "PluginManifest.h"
 #include "Misc/FileHelper.h"
+#include "ProfilingDebugging/CpuProfilerTrace.h"
 #include "Serialization/JsonReader.h"
 #include "Serialization/JsonSerializer.h"
 #include "JsonUtils/JsonConversion.h"
@@ -35,6 +36,7 @@ namespace PluginManifest
 
 TOptional<FText> UE::Projects::Private::Read(UE::Json::FConstObject Object, FPluginManifest& Result)
 {
+	TRACE_CPUPROFILER_EVENT_SCOPE_STR("UE::Projects::Private::Read FPluginManifest");
 	using namespace UE::Json;
 
 	TOptional<FConstArray> JsonContents = GetArrayField(Object, TEXT("Contents"));
@@ -78,12 +80,16 @@ TOptional<FText> UE::Projects::Private::Read(UE::Json::FConstObject Object, FPlu
 
 bool FPluginManifest::Load(const FString& FileName, FText& OutFailReason)
 {
+	TRACE_CPUPROFILER_EVENT_SCOPE_STR("FPluginManifest::Load");
 	// Read the file to a string
 	FString FileContents;
-	if (!FFileHelper::LoadFileToString(FileContents, *FileName))
 	{
-		OutFailReason = FText::Format(LOCTEXT("FailedToLoadDescriptorFile", "Failed to open descriptor file '{0}'"), FText::FromString(FileName));
-		return false;
+		TRACE_CPUPROFILER_EVENT_SCOPE_STR("LoadFileToString");
+		if (!FFileHelper::LoadFileToString(FileContents, *FileName))
+		{
+			OutFailReason = FText::Format(LOCTEXT("FailedToLoadDescriptorFile", "Failed to open descriptor file '{0}'"), FText::FromString(FileName));
+			return false;
+		}
 	}
 
 	{
@@ -99,23 +105,20 @@ bool FPluginManifest::Load(const FString& FileName, FText& OutFailReason)
 			ensure(FFileHelper::LoadFileToString(FileContents, *FileName));
 			OutFailReason = FText::Format(LOCTEXT("FailedToReadDescriptorFile", "Failed to read file. {0}"), FText::FromString(Document.GetError().CreateMessage(FileContents)));
 			
-		return false;
-	}
+			return false;
+		}
 
 		TOptional<UE::Json::FConstObject> RootObject = UE::Json::GetRootObject(Document.GetValue());
 		if (RootObject.IsSet())
 		{
-			//for (int32 N = 0; N < 1000; ++N)
-	{
-				if (TOptional<FText> Result = UE::Projects::Private::Read(*RootObject, *this))
-		{
-					OutFailReason = *Result;
-			return false;
+			if (TOptional<FText> Result = UE::Projects::Private::Read(*RootObject, *this))
+			{
+				OutFailReason = *Result;
+				return false;
+			}
 		}
-	}
 
-	return true;
-		}
+		return true;
 	}
 
 	return false;
