Made the EnumClassFlags.h functions UE_REWRITE so we never step into them.
Did a similar pass on TSharedPtr, TSharedRef and TWeakPtr.

diff --git a/Engine/Source/Runtime/Core/Public/Misc/EnumClassFlags.h b/Engine/Source/Runtime/Core/Public/Misc/EnumClassFlags.h
--- a/Engine/Source/Runtime/Core/Public/Misc/EnumClassFlags.h
+++ b/Engine/Source/Runtime/Core/Public/Misc/EnumClassFlags.h
@@ -2,16 +2,18 @@
 
 #pragma once
 
+#include "HAL/Platform.h"
+
 // Defines all bitwise operators for enum classes so it can be (mostly) used as a regular flags enum
 #define ENUM_CLASS_FLAGS(Enum) \
-	inline constexpr Enum& operator|=(Enum& Lhs, Enum Rhs) { return Lhs = (Enum)((__underlying_type(Enum))Lhs | (__underlying_type(Enum))Rhs); } \
-	inline constexpr Enum& operator&=(Enum& Lhs, Enum Rhs) { return Lhs = (Enum)((__underlying_type(Enum))Lhs & (__underlying_type(Enum))Rhs); } \
-	inline constexpr Enum& operator^=(Enum& Lhs, Enum Rhs) { return Lhs = (Enum)((__underlying_type(Enum))Lhs ^ (__underlying_type(Enum))Rhs); } \
-	inline constexpr Enum  operator| (Enum  Lhs, Enum Rhs) { return (Enum)((__underlying_type(Enum))Lhs | (__underlying_type(Enum))Rhs); } \
-	inline constexpr Enum  operator& (Enum  Lhs, Enum Rhs) { return (Enum)((__underlying_type(Enum))Lhs & (__underlying_type(Enum))Rhs); } \
-	inline constexpr Enum  operator^ (Enum  Lhs, Enum Rhs) { return (Enum)((__underlying_type(Enum))Lhs ^ (__underlying_type(Enum))Rhs); } \
-	inline constexpr bool  operator! (Enum  E)             { return !(__underlying_type(Enum))E; } \
-	inline constexpr Enum  operator~ (Enum  E)             { return (Enum)~(__underlying_type(Enum))E; }
+	UE_REWRITE constexpr Enum& operator|=(Enum& Lhs, Enum Rhs) { return Lhs = (Enum)((__underlying_type(Enum))Lhs | (__underlying_type(Enum))Rhs); } \
+	UE_REWRITE constexpr Enum& operator&=(Enum& Lhs, Enum Rhs) { return Lhs = (Enum)((__underlying_type(Enum))Lhs & (__underlying_type(Enum))Rhs); } \
+	UE_REWRITE constexpr Enum& operator^=(Enum& Lhs, Enum Rhs) { return Lhs = (Enum)((__underlying_type(Enum))Lhs ^ (__underlying_type(Enum))Rhs); } \
+	UE_REWRITE constexpr Enum  operator| (Enum  Lhs, Enum Rhs) { return (Enum)((__underlying_type(Enum))Lhs | (__underlying_type(Enum))Rhs); } \
+	UE_REWRITE constexpr Enum  operator& (Enum  Lhs, Enum Rhs) { return (Enum)((__underlying_type(Enum))Lhs & (__underlying_type(Enum))Rhs); } \
+	UE_REWRITE constexpr Enum  operator^ (Enum  Lhs, Enum Rhs) { return (Enum)((__underlying_type(Enum))Lhs ^ (__underlying_type(Enum))Rhs); } \
+	UE_REWRITE constexpr bool  operator! (Enum  E)             { return !(__underlying_type(Enum))E; } \
+	UE_REWRITE constexpr Enum  operator~ (Enum  E)             { return (Enum)~(__underlying_type(Enum))E; }
 
 // Friends all bitwise operators for enum classes so the definition can be kept private / protected.
 #define FRIEND_ENUM_CLASS_FLAGS(Enum) \
@@ -25,14 +27,14 @@
 	friend constexpr Enum  operator~ (Enum  E);
 
 template<typename Enum>
-constexpr bool EnumHasAllFlags(Enum Flags, Enum Contains)
+UE_REWRITE constexpr bool EnumHasAllFlags(Enum Flags, Enum Contains)
 {
 	using UnderlyingType = __underlying_type(Enum);
 	return ((UnderlyingType)Flags & (UnderlyingType)Contains) == (UnderlyingType)Contains;
 }
 
 template<typename Enum>
-constexpr bool EnumHasAnyFlags(Enum Flags, Enum Contains)
+UE_REWRITE constexpr bool EnumHasAnyFlags(Enum Flags, Enum Contains)
 {
 	using UnderlyingType = __underlying_type(Enum);
 	return ((UnderlyingType)Flags & (UnderlyingType)Contains) != 0;
@@ -52,7 +54,7 @@ constexpr bool EnumHasAnyFlags(Enum Flags, Enum Contains)
 * @note   Returns true if Flags is empty.
 */
 template <typename Enum>
-constexpr bool EnumOnlyContainsFlags(Enum Flags, Enum Contains)
+UE_REWRITE constexpr bool EnumOnlyContainsFlags(Enum Flags, Enum Contains)
 {
 	return EnumHasAllFlags(Contains, Flags);
 }
@@ -66,7 +68,7 @@ constexpr bool EnumOnlyContainsFlags(Enum Flags, Enum Contains)
 * @return true if Flags has only one flag set. false otherwise.
 */
 template <typename Enum>
-constexpr bool EnumHasOneFlag(Enum Flags)
+UE_REWRITE constexpr bool EnumHasOneFlag(Enum Flags)
 {
 	using UnderlyingType = __underlying_type(Enum);
 	return ((UnderlyingType)Flags != 0) && (((UnderlyingType)Flags & ((UnderlyingType)Flags - 1)) == 0);
@@ -81,28 +83,28 @@ constexpr bool EnumHasOneFlag(Enum Flags)
 * @return true if Flags has one and only one of the flags in OneOfFlags set. false otherwise.
 */
 template <typename Enum>
-constexpr bool EnumHasAnyOneFlag(Enum Flags, Enum OneOfFlags)
+UE_REWRITE constexpr bool EnumHasAnyOneFlag(Enum Flags, Enum OneOfFlags)
 {
 	using UnderlyingType = __underlying_type(Enum);
 	return EnumHasOneFlag((Enum)((UnderlyingType)Flags & (UnderlyingType)OneOfFlags));
 }
 
 template<typename Enum>
-constexpr void EnumAddFlags(Enum& Flags, Enum FlagsToAdd)
+UE_REWRITE constexpr void EnumAddFlags(Enum& Flags, Enum FlagsToAdd)
 {
 	using UnderlyingType = __underlying_type(Enum);
 	Flags = (Enum)((UnderlyingType)Flags | (UnderlyingType)FlagsToAdd);
 }
 
 template<typename Enum>
-constexpr void EnumRemoveFlags(Enum& Flags, Enum FlagsToRemove)
+UE_REWRITE constexpr void EnumRemoveFlags(Enum& Flags, Enum FlagsToRemove)
 {
 	using UnderlyingType = __underlying_type(Enum);
 	Flags = (Enum)((UnderlyingType)Flags & ~(UnderlyingType)FlagsToRemove);
 }
 
 template <typename Enum>
-constexpr Enum EnumLowestSetFlag(Enum Flags)
+UE_REWRITE constexpr Enum EnumLowestSetFlag(Enum Flags)
 {
 	using UnderlyingType = __underlying_type(Enum);
 #ifdef _MSC_VER
@@ -116,14 +118,14 @@ constexpr Enum EnumLowestSetFlag(Enum Flags)
 }
 
 template <typename Enum>
-constexpr Enum EnumRemoveLowestSetFlag(Enum Flags)
+UE_REWRITE constexpr Enum EnumRemoveLowestSetFlag(Enum Flags)
 {
 	using UnderlyingType = __underlying_type(Enum);
 	return (Enum)((UnderlyingType)Flags & ((UnderlyingType)Flags - 1));
 }
 
 template <typename Enum>
-constexpr int EnumNumSetFlags(Enum Flags)
+UE_NODEBUG constexpr int EnumNumSetFlags(Enum Flags)
 {
 	using UnderlyingType = __underlying_type(Enum);
 
diff --git a/Engine/Source/Runtime/Core/Public/Templates/SharedPointer.h b/Engine/Source/Runtime/Core/Public/Templates/SharedPointer.h
--- a/Engine/Source/Runtime/Core/Public/Templates/SharedPointer.h
+++ b/Engine/Source/Runtime/Core/Public/Templates/SharedPointer.h
@@ -170,7 +170,7 @@ public:
 		typename OtherType
 		UE_REQUIRES(std::is_convertible_v<OtherType*, ObjectType*>)
 	>
-	inline explicit TSharedRef( OtherType* InObject )
+	UE_NODEBUG explicit TSharedRef( OtherType* InObject )
 		: Object( InObject )
 		, SharedReferenceCount( SharedPointerInternals::NewDefaultReferenceController< Mode >( InObject ) )
 	{
@@ -188,7 +188,7 @@ public:
 		typename DeleterType
 		UE_REQUIRES(std::is_convertible_v<OtherType*, ObjectType*>)
 	>
-	inline TSharedRef( OtherType* InObject, DeleterType&& InDeleter )
+	UE_NODEBUG inline TSharedRef( OtherType* InObject, DeleterType&& InDeleter )
 		: Object( InObject )
 		, SharedReferenceCount( SharedPointerInternals::NewCustomReferenceController< Mode >( InObject, Forward< DeleterType >( InDeleter ) ) )
 	{
@@ -219,7 +219,7 @@ public:
 		typename OtherType
 		UE_REQUIRES(std::is_convertible_v<OtherType*, ObjectType*>)
 	>
-	inline TSharedRef( SharedPointerInternals::TRawPtrProxy< OtherType > const& InRawPtrProxy )
+	UE_NODEBUG inline TSharedRef( SharedPointerInternals::TRawPtrProxy< OtherType > const& InRawPtrProxy )
 		: Object( InRawPtrProxy.Object )
 		, SharedReferenceCount( SharedPointerInternals::NewDefaultReferenceController< Mode >( InRawPtrProxy.Object ) )
 	{
@@ -244,7 +244,7 @@ public:
 		typename DeleterType
 		UE_REQUIRES(std::is_convertible_v<OtherType*, ObjectType*>)
 	>
-	inline TSharedRef( SharedPointerInternals::TRawPtrProxyWithDeleter< OtherType, DeleterType > const& InRawPtrProxy )
+	UE_NODEBUG inline TSharedRef( SharedPointerInternals::TRawPtrProxyWithDeleter< OtherType, DeleterType > const& InRawPtrProxy )
 		: Object( InRawPtrProxy.Object )
 		, SharedReferenceCount( SharedPointerInternals::NewCustomReferenceController< Mode >( InRawPtrProxy.Object, InRawPtrProxy.Deleter ) )
 	{
@@ -269,7 +269,7 @@ public:
 		typename DeleterType
 		UE_REQUIRES(std::is_convertible_v<OtherType*, ObjectType*>)
 	>
-	inline TSharedRef( SharedPointerInternals::TRawPtrProxyWithDeleter< OtherType, DeleterType >&& InRawPtrProxy )
+	UE_NODEBUG inline TSharedRef( SharedPointerInternals::TRawPtrProxyWithDeleter< OtherType, DeleterType >&& InRawPtrProxy )
 		: Object( InRawPtrProxy.Object )
 		, SharedReferenceCount( SharedPointerInternals::NewCustomReferenceController< Mode >( InRawPtrProxy.Object, MoveTemp( InRawPtrProxy.Deleter ) ) )
 	{
@@ -292,7 +292,7 @@ public:
 		typename OtherType
 		UE_REQUIRES(std::is_convertible_v<OtherType*, ObjectType*>)
 	>
-	inline TSharedRef( TSharedRef< OtherType, Mode > const& InSharedRef )
+	UE_NODEBUG inline TSharedRef( TSharedRef< OtherType, Mode > const& InSharedRef )
 		: Object( InSharedRef.Object )
 		, SharedReferenceCount( InSharedRef.SharedReferenceCount )
 	{
@@ -307,7 +307,7 @@ public:
 	 * @param  InSharedRef  The shared reference whose object we should create an additional reference to
 	 */
 	template <typename OtherType>
-	inline TSharedRef( TSharedRef< OtherType, Mode > const& InSharedRef, SharedPointerInternals::FStaticCastTag )
+	UE_NODEBUG inline TSharedRef( TSharedRef< OtherType, Mode > const& InSharedRef, SharedPointerInternals::FStaticCastTag )
 		: Object( static_cast< ObjectType* >( InSharedRef.Object ) )
 		, SharedReferenceCount( InSharedRef.SharedReferenceCount )
 	{
@@ -322,7 +322,7 @@ public:
 	 * @param  InSharedRef  The shared reference whose object we should create an additional reference to
 	 */
 	template <typename OtherType>
-	inline TSharedRef( TSharedRef< OtherType, Mode > const& InSharedRef, SharedPointerInternals::FConstCastTag )
+	UE_NODEBUG inline TSharedRef( TSharedRef< OtherType, Mode > const& InSharedRef, SharedPointerInternals::FConstCastTag )
 		: Object( const_cast< ObjectType* >( InSharedRef.Object ) )
 		, SharedReferenceCount( InSharedRef.SharedReferenceCount )
 	{
@@ -336,7 +336,7 @@ public:
 	 * @param  InObject  The object pointer to use (instead of the incoming shared pointer's object)
 	 */
 	template <typename OtherType>
-	inline TSharedRef( TSharedRef< OtherType, Mode > const& OtherSharedRef, ObjectType* InObject )
+	UE_NODEBUG inline TSharedRef( TSharedRef< OtherType, Mode > const& OtherSharedRef, ObjectType* InObject )
 		: Object( InObject )
 		, SharedReferenceCount( OtherSharedRef.SharedReferenceCount )
 	{
@@ -345,12 +345,13 @@ public:
 		check( InObject != nullptr );
 	}
 
-	inline TSharedRef( TSharedRef const& InSharedRef )
+	UE_NODEBUG inline TSharedRef( TSharedRef const& InSharedRef )
 		: Object( InSharedRef.Object )
 		, SharedReferenceCount( InSharedRef.SharedReferenceCount )
-	{ }
+	{
+	}
 
-	inline TSharedRef( TSharedRef&& InSharedRef )
+	UE_NODEBUG inline TSharedRef( TSharedRef&& InSharedRef )
 		: Object( InSharedRef.Object )
 		, SharedReferenceCount( InSharedRef.SharedReferenceCount )
 	{
@@ -366,14 +367,14 @@ public:
 	 *
 	 * @param  InSharedRef  Shared reference to replace with
 	 */
-	inline TSharedRef& operator=( TSharedRef const& InSharedRef )
+	UE_NODEBUG inline TSharedRef& operator=( TSharedRef const& InSharedRef )
 	{
 		TSharedRef Temp = InSharedRef;
 		::Swap(Temp, *this);
 		return *this;
 	}
 
-	inline TSharedRef& operator=( TSharedRef&& InSharedRef )
+	UE_REWRITE TSharedRef& operator=( TSharedRef&& InSharedRef )
 	{
 		::Swap(*this, InSharedRef);
 		return *this;
@@ -391,7 +392,7 @@ public:
 		typename OtherType
 		UE_REQUIRES(std::is_convertible_v<OtherType*, ObjectType*>)
 	>
-	inline TSharedRef& operator=( SharedPointerInternals::TRawPtrProxy< OtherType > const& InRawPtrProxy )
+	UE_NODEBUG inline TSharedRef& operator=( SharedPointerInternals::TRawPtrProxy< OtherType > const& InRawPtrProxy )
 	{
 		// If the following assert goes off, it means a TSharedRef was initialized from a nullptr object pointer.
 		// Shared references must never be nullptr, so either pass a valid object or consider using TSharedPtr instead.
@@ -414,7 +415,7 @@ public:
 		typename DeleterType
 		UE_REQUIRES(std::is_convertible_v<OtherType*, ObjectType*>)
 	>
-	inline TSharedRef& operator=( SharedPointerInternals::TRawPtrProxyWithDeleter< OtherType, DeleterType > const& InRawPtrProxy )
+	UE_NODEBUG inline TSharedRef& operator=( SharedPointerInternals::TRawPtrProxyWithDeleter< OtherType, DeleterType > const& InRawPtrProxy )
 	{
 		// If the following assert goes off, it means a TSharedRef was initialized from a nullptr object pointer.
 		// Shared references must never be nullptr, so either pass a valid object or consider using TSharedPtr instead.
@@ -437,7 +438,7 @@ public:
 		typename DeleterType
 		UE_REQUIRES(std::is_convertible_v<OtherType*, ObjectType*>)
 	>
-	inline TSharedRef& operator=( SharedPointerInternals::TRawPtrProxyWithDeleter< OtherType, DeleterType >&& InRawPtrProxy )
+	UE_NODEBUG inline TSharedRef& operator=( SharedPointerInternals::TRawPtrProxyWithDeleter< OtherType, DeleterType >&& InRawPtrProxy )
 	{
 		// If the following assert goes off, it means a TSharedRef was initialized from a nullptr object pointer.
 		// Shared references must never be nullptr, so either pass a valid object or consider using TSharedPtr instead.
@@ -455,7 +456,7 @@ public:
 	 *
 	 * @return  Shared pointer to the object
 	 */
-	[[nodiscard]] UE_FORCEINLINE_HINT TSharedPtr<ObjectType, Mode> ToSharedPtr() const
+	[[nodiscard]] UE_REWRITE TSharedPtr<ObjectType, Mode> ToSharedPtr() const
 	{
 		return TSharedPtr<ObjectType, Mode>(*this);
 	}
@@ -465,7 +466,7 @@ public:
 	 *
 	 * @return  Weak pointer to the object
 	 */
-	[[nodiscard]] UE_FORCEINLINE_HINT TWeakPtr<ObjectType, Mode> ToWeakPtr() const
+	[[nodiscard]] UE_REWRITE TWeakPtr<ObjectType, Mode> ToWeakPtr() const
 	{
 		return TWeakPtr<ObjectType, Mode>(*this);
 	}
@@ -475,7 +476,7 @@ public:
 	 *
 	 * @return  The object owned by this shared reference
 	 */
-	[[nodiscard]] inline ObjectType& Get() const
+	[[nodiscard]] UE_REWRITE ObjectType& Get() const
 	{
 		// Should never be nullptr as TSharedRef is never nullable
 		checkSlow( IsValid() );
@@ -487,7 +488,7 @@ public:
 	 *
 	 * @return  Reference to the object
 	 */
-	[[nodiscard]] inline ObjectType& operator*() const
+	[[nodiscard]] UE_REWRITE ObjectType& operator*() const
 	{
 		// Should never be nullptr as TSharedRef is never nullable
 		checkSlow( IsValid() );
@@ -499,7 +500,7 @@ public:
 	 *
 	 * @return  Returns a pointer to the object referenced by this shared reference
 	 */
-	[[nodiscard]] inline ObjectType* operator->() const
+	[[nodiscard]] UE_REWRITE ObjectType* operator->() const
 	{
 		// Should never be nullptr as TSharedRef is never nullable
 		checkSlow( IsValid() );
@@ -512,7 +513,7 @@ public:
 	 *
 	 * @return  Number of shared references to the object (including this reference.)
 	 */
-	[[nodiscard]] UE_FORCEINLINE_HINT int32 GetSharedReferenceCount() const
+	[[nodiscard]] UE_REWRITE int32 GetSharedReferenceCount() const
 	{
 		return SharedReferenceCount.GetSharedReferenceCount();
 	}
@@ -528,7 +529,7 @@ public:
 	 *
 	 * @return  True if there is only one shared reference to the object, and this is it!
 	 */
-	[[nodiscard]] UE_FORCEINLINE_HINT bool IsUnique() const
+	[[nodiscard]] UE_REWRITE bool IsUnique() const
 	{
 		return SharedReferenceCount.IsUnique();
 	}
@@ -539,11 +540,11 @@ public:
 	constexpr static bool bHasIntrusiveUnsetOptionalState = true;
 	using IntrusiveUnsetOptionalStateType = TSharedRef;
 
-	explicit TSharedRef(FIntrusiveUnsetOptionalState)
+	UE_REWRITE explicit TSharedRef(FIntrusiveUnsetOptionalState)
 		: Object(nullptr)
 	{
 	}
-	bool UEOpEquals(FIntrusiveUnsetOptionalState) const
+	UE_REWRITE bool UEOpEquals(FIntrusiveUnsetOptionalState) const
 	{
 		return !IsValid();
 	}
@@ -552,13 +553,13 @@ public:
 	/////////////////////////////////////////////////
 
 	template <typename OtherType>
-	[[nodiscard]] bool UEOpEquals(const TSharedRef<OtherType, Mode>& Rhs) const
+	[[nodiscard]] UE_NODEBUG bool UEOpEquals(const TSharedRef<OtherType, Mode>& Rhs) const
 	{
 		return this->Object == &Rhs.Get();
 	}
 
 	template <typename OtherType>
-	[[nodiscard]] bool UEOpEquals(const TSharedPtr<OtherType, Mode>& Rhs) const
+	[[nodiscard]] UE_NODEBUG bool UEOpEquals(const TSharedPtr<OtherType, Mode>& Rhs) const
 	{
 		// This comparison against null is maintained as existing behavior, but isn't consistent with TWeakPtr comparison.
 		OtherType* RhsPtr = Rhs.Get();
@@ -566,13 +567,13 @@ public:
 	}
 
 	template <typename OtherType>
-	[[nodiscard]] bool UEOpEquals(const TWeakPtr<OtherType, Mode>& Rhs) const
+	[[nodiscard]] UE_NODEBUG bool UEOpEquals(const TWeakPtr<OtherType, Mode>& Rhs) const
 	{
 		return this->Object == Rhs.Pin().Get();
 	}
 
 	template <typename OtherType>
-	[[nodiscard]] bool UEOpLessThan(const TSharedRef<OtherType, Mode>& Rhs) const
+	[[nodiscard]] UE_NODEBUG bool UEOpLessThan(const TSharedRef<OtherType, Mode>& Rhs) const
 	{
 		return this->Object < &Rhs.Get();
 	}
@@ -630,7 +631,7 @@ private:
 	 *
 	 * @return  True if the shared reference is valid and can be dereferenced
 	 */
-	[[nodiscard]] UE_FORCEINLINE_HINT bool IsValid() const
+	[[nodiscard]] UE_REWRITE bool IsValid() const
 	{
 		return Object != nullptr;
 	}
@@ -663,7 +664,7 @@ private:
 		#pragma warning(pop)
 	#endif
 
-	inline explicit TSharedRef(ObjectType* InObject, SharedPointerInternals::TReferenceControllerBase<Mode>* InSharedReferenceCount)
+	UE_NODEBUG inline explicit TSharedRef(ObjectType* InObject, SharedPointerInternals::TReferenceControllerBase<Mode>* InSharedReferenceCount)
 		: Object(InObject)
 		, SharedReferenceCount(InSharedReferenceCount)
 	{
@@ -702,7 +703,7 @@ public:
 	 * Constructs an empty shared pointer
 	 */
 	// NOTE: FNullTag parameter is an Unreal extension to standard shared_ptr behavior
-	inline TSharedPtr( SharedPointerInternals::FNullTag* = nullptr )
+	UE_REWRITE TSharedPtr( SharedPointerInternals::FNullTag* = nullptr )
 		: Object( nullptr )
 		, SharedReferenceCount()
 	{
@@ -718,7 +719,7 @@ public:
 		typename OtherType
 		UE_REQUIRES(std::is_convertible_v<OtherType*, ObjectType*>)
 	>
-	inline explicit TSharedPtr( OtherType* InObject )
+	UE_NODEBUG inline explicit TSharedPtr( OtherType* InObject )
 		: Object( InObject )
 		, SharedReferenceCount( SharedPointerInternals::NewDefaultReferenceController< Mode >( InObject ) )
 	{
@@ -739,7 +740,7 @@ public:
 		typename DeleterType
 		UE_REQUIRES(std::is_convertible_v<OtherType*, ObjectType*>)
 	>
-	inline TSharedPtr( OtherType* InObject, DeleterType&& InDeleter )
+	UE_NODEBUG inline TSharedPtr( OtherType* InObject, DeleterType&& InDeleter )
 		: Object( InObject )
 		, SharedReferenceCount( SharedPointerInternals::NewCustomReferenceController< Mode >( InObject, Forward< DeleterType >( InDeleter ) ) )
 	{
@@ -758,7 +759,7 @@ public:
 		typename OtherType
 		UE_REQUIRES(std::is_convertible_v<OtherType*, ObjectType*>)
 	>
-	inline TSharedPtr( SharedPointerInternals::TRawPtrProxy< OtherType > const& InRawPtrProxy )
+	UE_NODEBUG inline TSharedPtr( SharedPointerInternals::TRawPtrProxy< OtherType > const& InRawPtrProxy )
 		: Object( InRawPtrProxy.Object )
 		, SharedReferenceCount( SharedPointerInternals::NewDefaultReferenceController< Mode >( InRawPtrProxy.Object ) )
 	{
@@ -778,7 +779,7 @@ public:
 		typename DeleterType
 		UE_REQUIRES(std::is_convertible_v<OtherType*, ObjectType*>)
 	>
-	inline TSharedPtr( SharedPointerInternals::TRawPtrProxyWithDeleter< OtherType, DeleterType > const& InRawPtrProxy )
+	UE_NODEBUG inline TSharedPtr( SharedPointerInternals::TRawPtrProxyWithDeleter< OtherType, DeleterType > const& InRawPtrProxy )
 		: Object( InRawPtrProxy.Object )
 		, SharedReferenceCount( SharedPointerInternals::NewCustomReferenceController< Mode >( InRawPtrProxy.Object, InRawPtrProxy.Deleter ) )
 	{
@@ -798,7 +799,7 @@ public:
 		typename DeleterType
 		UE_REQUIRES(std::is_convertible_v<OtherType*, ObjectType*>)
 	>
-	inline TSharedPtr( SharedPointerInternals::TRawPtrProxyWithDeleter< OtherType, DeleterType >&& InRawPtrProxy )
+	UE_NODEBUG inline TSharedPtr( SharedPointerInternals::TRawPtrProxyWithDeleter< OtherType, DeleterType >&& InRawPtrProxy )
 		: Object( InRawPtrProxy.Object )
 		, SharedReferenceCount( SharedPointerInternals::NewCustomReferenceController< Mode >( InRawPtrProxy.Object, MoveTemp( InRawPtrProxy.Deleter ) ) )
 	{
@@ -817,19 +818,19 @@ public:
 		typename OtherType
 		UE_REQUIRES(std::is_convertible_v<OtherType*, ObjectType*>)
 	>
-	inline TSharedPtr( TSharedPtr< OtherType, Mode > const& InSharedPtr )
+	UE_NODEBUG inline TSharedPtr( TSharedPtr< OtherType, Mode > const& InSharedPtr )
 		: Object( InSharedPtr.Object )
 		, SharedReferenceCount( InSharedPtr.SharedReferenceCount )
 	{
 	}
 
-	inline TSharedPtr( TSharedPtr const& InSharedPtr )
+	UE_NODEBUG inline TSharedPtr( TSharedPtr const& InSharedPtr )
 		: Object( InSharedPtr.Object )
 		, SharedReferenceCount( InSharedPtr.SharedReferenceCount )
 	{
 	}
 
-	inline TSharedPtr( TSharedPtr&& InSharedPtr )
+	UE_NODEBUG inline TSharedPtr( TSharedPtr&& InSharedPtr )
 		: Object( InSharedPtr.Object )
 		, SharedReferenceCount( MoveTemp(InSharedPtr.SharedReferenceCount) )
 	{
@@ -847,7 +848,7 @@ public:
 		typename OtherType
 		UE_REQUIRES(std::is_convertible_v<OtherType*, ObjectType*>)
 	>
-	inline TSharedPtr( TSharedRef< OtherType, Mode > const& InSharedRef )
+	UE_NODEBUG inline TSharedPtr( TSharedRef< OtherType, Mode > const& InSharedRef )
 		: Object( InSharedRef.Object )
 		, SharedReferenceCount( InSharedRef.SharedReferenceCount )
 	{
@@ -864,7 +865,7 @@ public:
 	 * @param  InSharedPtr  The shared pointer whose object we should create an additional reference to
 	 */
 	template <typename OtherType>
-	inline TSharedPtr( TSharedPtr< OtherType, Mode > const& InSharedPtr, SharedPointerInternals::FStaticCastTag )
+	UE_NODEBUG inline TSharedPtr( TSharedPtr< OtherType, Mode > const& InSharedPtr, SharedPointerInternals::FStaticCastTag )
 		: Object( static_cast< ObjectType* >( InSharedPtr.Object ) )
 		, SharedReferenceCount( InSharedPtr.SharedReferenceCount )
 	{
@@ -879,7 +880,7 @@ public:
 	 * @param  InSharedPtr  The shared pointer whose object we should create an additional reference to
 	 */
 	template <typename OtherType>
-	inline TSharedPtr( TSharedPtr< OtherType, Mode > const& InSharedPtr, SharedPointerInternals::FConstCastTag )
+	UE_NODEBUG inline TSharedPtr( TSharedPtr< OtherType, Mode > const& InSharedPtr, SharedPointerInternals::FConstCastTag )
 		: Object( const_cast< ObjectType* >( InSharedPtr.Object ) )
 		, SharedReferenceCount( InSharedPtr.SharedReferenceCount )
 	{
@@ -893,7 +894,7 @@ public:
 	 * @param  InObject  The object pointer to use (instead of the incoming shared pointer's object)
 	 */
 	template <typename OtherType>
-	inline TSharedPtr( TSharedPtr< OtherType, Mode > const& OtherSharedPtr, ObjectType* InObject )
+	UE_NODEBUG inline TSharedPtr( TSharedPtr< OtherType, Mode > const& OtherSharedPtr, ObjectType* InObject )
 		: Object( InObject )
 		, SharedReferenceCount( OtherSharedPtr.SharedReferenceCount )
 	{
@@ -907,7 +908,7 @@ public:
 	 * @param  InObject  The object pointer to use (instead of the incoming shared pointer's object)
 	 */
 	template <typename OtherType>
-	inline TSharedPtr( TSharedPtr< OtherType, Mode >&& OtherSharedPtr, ObjectType* InObject )
+	UE_NODEBUG inline TSharedPtr( TSharedPtr< OtherType, Mode >&& OtherSharedPtr, ObjectType* InObject )
 		: Object( InObject )
 		, SharedReferenceCount( MoveTemp(OtherSharedPtr.SharedReferenceCount) )
 	{
@@ -922,7 +923,7 @@ public:
 	 * @param  InObject  The object pointer to use (instead of the incoming shared pointer's object)
 	 */
 	template <typename OtherType>
-	inline TSharedPtr( TSharedRef< OtherType, Mode > const& OtherSharedRef, ObjectType* InObject )
+	UE_NODEBUG inline TSharedPtr( TSharedRef< OtherType, Mode > const& OtherSharedRef, ObjectType* InObject )
 		: Object( InObject )
 		, SharedReferenceCount( OtherSharedRef.SharedReferenceCount )
 	{
@@ -933,7 +934,7 @@ public:
 	 * referenced and will be deleted if there are no other referencers.
 	 */
 	// NOTE: The following is an Unreal extension to standard shared_ptr behavior
-	inline TSharedPtr& operator=( SharedPointerInternals::FNullTag* )
+	UE_REWRITE TSharedPtr& operator=( SharedPointerInternals::FNullTag* )
 	{
 		Reset();
 		return *this;
@@ -946,7 +947,7 @@ public:
 	 *
 	 * @param  InSharedPtr  Shared pointer to replace with
 	 */
-	inline TSharedPtr& operator=( TSharedPtr const& InSharedPtr )
+	UE_NODEBUG inline TSharedPtr& operator=( TSharedPtr const& InSharedPtr )
 	{
 		TSharedPtr Temp = InSharedPtr;
 		::Swap(Temp, *this);
@@ -955,7 +956,7 @@ public:
 
 	// Disable false positive buffer overrun warning during pgo linking step
 	PGO_LINK_DISABLE_WARNINGS
-	inline TSharedPtr& operator=( TSharedPtr&& InSharedPtr )
+	UE_NODEBUG inline TSharedPtr& operator=( TSharedPtr&& InSharedPtr )
 	{
 		if (this != &InSharedPtr)
 		{
@@ -979,7 +980,7 @@ public:
 		typename OtherType
 		UE_REQUIRES(std::is_convertible_v<OtherType*, ObjectType*>)
 	>
-	inline TSharedPtr& operator=( SharedPointerInternals::TRawPtrProxy< OtherType > const& InRawPtrProxy )
+	UE_REWRITE TSharedPtr& operator=( SharedPointerInternals::TRawPtrProxy< OtherType > const& InRawPtrProxy )
 	{
 		*this = TSharedPtr< ObjectType, Mode >( InRawPtrProxy );
 		return *this;
@@ -998,7 +999,7 @@ public:
 		typename DeleterType
 		UE_REQUIRES(std::is_convertible_v<OtherType*, ObjectType*>)
 	>
-	inline TSharedPtr& operator=( SharedPointerInternals::TRawPtrProxyWithDeleter< OtherType, DeleterType > const& InRawPtrProxy )
+	UE_REWRITE TSharedPtr& operator=( SharedPointerInternals::TRawPtrProxyWithDeleter< OtherType, DeleterType > const& InRawPtrProxy )
 	{
 		*this = TSharedPtr< ObjectType, Mode >( InRawPtrProxy );
 		return *this;
@@ -1017,7 +1018,7 @@ public:
 		typename DeleterType
 		UE_REQUIRES(std::is_convertible_v<OtherType*, ObjectType*>)
 	>
-	inline TSharedPtr& operator=( SharedPointerInternals::TRawPtrProxyWithDeleter< OtherType, DeleterType >&& InRawPtrProxy )
+	UE_REWRITE TSharedPtr& operator=( SharedPointerInternals::TRawPtrProxyWithDeleter< OtherType, DeleterType >&& InRawPtrProxy )
 	{
 		*this = TSharedPtr< ObjectType, Mode >( MoveTemp( InRawPtrProxy ) );
 		return *this;
@@ -1032,7 +1033,7 @@ public:
 	 * @return  Reference to the object
 	 */
 	// NOTE: The following is an Unreal extension to standard shared_ptr behavior
-	[[nodiscard]] inline TSharedRef< ObjectType, Mode > ToSharedRef() const&
+	[[nodiscard]] UE_NODEBUG inline TSharedRef< ObjectType, Mode > ToSharedRef() const&
 	{
 		// If this assert goes off, it means a shared reference was created from a shared pointer that was nullptr.
 		// Shared references are never allowed to be null.  Consider using TSharedPtr instead.
@@ -1046,7 +1047,7 @@ public:
 	 * @return  Reference to the object
 	 */
 	 // NOTE: The following is an Unreal extension to standard shared_ptr behavior
-	[[nodiscard]] inline TSharedRef< ObjectType, Mode > ToSharedRef() &&
+	[[nodiscard]] UE_NODEBUG inline TSharedRef< ObjectType, Mode > ToSharedRef() &&
 	{
 		// If this assert goes off, it means a shared reference was created from a shared pointer that was nullptr.
 		// Shared references are never allowed to be null.  Consider using TSharedPtr instead.
@@ -1059,7 +1060,7 @@ public:
 	 *
 	 * @return  Weak pointer to the object
 	 */
-	[[nodiscard]] UE_FORCEINLINE_HINT TWeakPtr<ObjectType, Mode> ToWeakPtr() const
+	[[nodiscard]] UE_REWRITE TWeakPtr<ObjectType, Mode> ToWeakPtr() const
 	{
 		return TWeakPtr<ObjectType, Mode>(*this);
 	}
@@ -1069,7 +1070,7 @@ public:
 	 *
 	 * @return  The object owned by this shared pointer, or nullptr
 	 */
-	[[nodiscard]] UE_FORCEINLINE_HINT ObjectType* Get() const
+	[[nodiscard]] UE_REWRITE ObjectType* Get() const
 	{
 		return Object;
 	}
@@ -1079,7 +1080,7 @@ public:
 	 *
 	 * @return  True if the shared pointer is valid and can be dereferenced
 	 */
-	[[nodiscard]] UE_FORCEINLINE_HINT explicit operator bool() const
+	[[nodiscard]] UE_REWRITE explicit operator bool() const
 	{
 		return Object != nullptr;
 	}
@@ -1089,7 +1090,7 @@ public:
 	 *
 	 * @return  True if the shared pointer is valid and can be dereferenced
 	 */
-	[[nodiscard]] UE_FORCEINLINE_HINT const bool IsValid() const
+	[[nodiscard]] UE_REWRITE const bool IsValid() const
 	{
 		return Object != nullptr;
 	}
@@ -1099,11 +1100,8 @@ public:
 	 *
 	 * @return  Reference to the object
 	 */
-	template <
-		typename DummyObjectType = ObjectType
-		UE_REQUIRES(UE_REQUIRES_EXPR(*(DummyObjectType*)nullptr)) // this construct means that operator* is only considered for overload resolution if T is dereferenceable
-	>
-	[[nodiscard]] inline DummyObjectType& operator*() const
+	[[nodiscard]] UE_NODEBUG inline std::add_lvalue_reference_t<ObjectType> operator*() const
+		requires requires (ObjectType* Ptr) { *Ptr; }
 	{
 		check( IsValid() );
 		return *Object;
@@ -1114,7 +1112,7 @@ public:
 	 *
 	 * @return  Returns a pointer to the object referenced by this shared pointer
 	 */
-	[[nodiscard]] inline ObjectType* operator->() const
+	[[nodiscard]] UE_NODEBUG inline ObjectType* operator->() const
 	{
 		check( IsValid() );
 		return Object;
@@ -1124,7 +1122,7 @@ public:
 	 * Resets this shared pointer, removing a reference to the object.  If there are no other shared
 	 * references to the object then it will be destroyed.
 	 */
-	UE_FORCEINLINE_HINT void Reset()
+	UE_REWRITE void Reset()
 	{
  		*this = TSharedPtr< ObjectType, Mode >();
 	}
@@ -1135,7 +1133,7 @@ public:
 	 *
 	 * @return  Number of shared references to the object (including this reference.)
 	 */
-	[[nodiscard]] UE_FORCEINLINE_HINT int32 GetSharedReferenceCount() const
+	[[nodiscard]] UE_REWRITE int32 GetSharedReferenceCount() const
 	{
 		return SharedReferenceCount.GetSharedReferenceCount();
 	}
@@ -1147,7 +1145,7 @@ public:
 	 *
 	 * @return  True if there is only one shared reference to the object, and this is it!
 	 */
-	[[nodiscard]] UE_FORCEINLINE_HINT bool IsUnique() const
+	[[nodiscard]] UE_REWRITE bool IsUnique() const
 	{
 		return SharedReferenceCount.IsUnique();
 	}
@@ -1158,11 +1156,11 @@ public:
 	constexpr static bool bHasIntrusiveUnsetOptionalState = true;
 	using IntrusiveUnsetOptionalStateType = TSharedPtr;
 
-	explicit TSharedPtr(FIntrusiveUnsetOptionalState)
+	UE_REWRITE explicit TSharedPtr(FIntrusiveUnsetOptionalState)
 		: Object((ObjectType*)-1)
 	{
 	}
-	bool UEOpEquals(FIntrusiveUnsetOptionalState) const
+	UE_REWRITE bool UEOpEquals(FIntrusiveUnsetOptionalState) const
 	{
 		return Object == (ObjectType*)-1;
 	}
@@ -1171,7 +1169,7 @@ public:
 	/////////////////////////////////////////////////
 
 	template <typename OtherType>
-	[[nodiscard]] inline bool UEOpEquals(const TSharedPtr<OtherType, Mode>& Rhs) const
+	[[nodiscard]] UE_REWRITE bool UEOpEquals(const TSharedPtr<OtherType, Mode>& Rhs) const
 	{
 		return this->Object == Rhs.Get();
 	}
@@ -1188,7 +1186,7 @@ public:
 	}
 
 	template <typename OtherType>
-	[[nodiscard]] inline bool UEOpLessThan(const TSharedPtr<OtherType, Mode>& Rhs) const
+	[[nodiscard]] UE_REWRITE bool UEOpLessThan(const TSharedPtr<OtherType, Mode>& Rhs) const
 	{
 		return this->Object < Rhs.Get();
 	}
@@ -1207,7 +1205,7 @@ private:
 		typename OtherType
 		UE_REQUIRES(std::is_convertible_v<OtherType*, ObjectType*>)
 	>
-	inline explicit TSharedPtr( TWeakPtr< OtherType, Mode > const& InWeakPtr )
+	UE_NODEBUG inline explicit TSharedPtr( TWeakPtr< OtherType, Mode > const& InWeakPtr )
 		: Object( nullptr )
 		, SharedReferenceCount( InWeakPtr.WeakReferenceCount )
 	{
@@ -1232,7 +1230,7 @@ private:
 		typename OtherType
 		UE_REQUIRES(std::is_convertible_v<OtherType*, ObjectType*>)
 	>
-	inline explicit TSharedPtr(TWeakPtr< OtherType, Mode >&& InWeakPtr)
+	UE_NODEBUG inline explicit TSharedPtr(TWeakPtr< OtherType, Mode >&& InWeakPtr)
 		: Object(nullptr)
 		, SharedReferenceCount( MoveTemp( InWeakPtr.WeakReferenceCount ) )
 	{
@@ -1306,7 +1304,7 @@ public:
 
 	/** Constructs an empty TWeakPtr */
 	// NOTE: FNullTag parameter is an Unreal extension to standard shared_ptr behavior
-	inline TWeakPtr( SharedPointerInternals::FNullTag* = nullptr )
+	UE_NODEBUG inline TWeakPtr( SharedPointerInternals::FNullTag* = nullptr )
 		: Object( nullptr )
 		, WeakReferenceCount()
 	{
@@ -1322,7 +1320,7 @@ public:
 		typename OtherType
 		UE_REQUIRES(std::is_convertible_v<OtherType*, ObjectType*>)
 	>
-	inline TWeakPtr( TSharedRef< OtherType, Mode > const& InSharedRef )
+	UE_NODEBUG inline TWeakPtr( TSharedRef< OtherType, Mode > const& InSharedRef )
 		: Object( InSharedRef.Object )
 		, WeakReferenceCount( InSharedRef.SharedReferenceCount )
 	{
@@ -1337,7 +1335,7 @@ public:
 		typename OtherType
 		UE_REQUIRES(std::is_convertible_v<OtherType*, ObjectType*>)
 	>
-	inline TWeakPtr( TSharedPtr< OtherType, Mode > const& InSharedPtr )
+	UE_NODEBUG inline TWeakPtr( TSharedPtr< OtherType, Mode > const& InSharedPtr )
 		: Object( InSharedPtr.Object )
 		, WeakReferenceCount( InSharedPtr.SharedReferenceCount )
 	{
@@ -1352,7 +1350,7 @@ public:
 	 * @param  InWeakPtr  The weak pointer whose object we should create an additional reference to
 	 */
 	template <typename OtherType>
-	inline TWeakPtr(TWeakPtr< OtherType, Mode > const& InWeakPtr, SharedPointerInternals::FStaticCastTag)
+	UE_NODEBUG inline TWeakPtr(TWeakPtr< OtherType, Mode > const& InWeakPtr, SharedPointerInternals::FStaticCastTag)
 		: Object(static_cast<ObjectType*>(InWeakPtr.Object))
 		, WeakReferenceCount(InWeakPtr.WeakReferenceCount)
 	{
@@ -1367,7 +1365,7 @@ public:
 	 * @param  InWeakPtr  The weak pointer whose object we should create an additional reference to
 	 */
 	template <typename OtherType>
-	inline TWeakPtr(TWeakPtr< OtherType, Mode > const& InWeakPtr, SharedPointerInternals::FConstCastTag)
+	UE_NODEBUG inline TWeakPtr(TWeakPtr< OtherType, Mode > const& InWeakPtr, SharedPointerInternals::FConstCastTag)
 		: Object(const_cast<ObjectType*>(InWeakPtr.Object))
 		, WeakReferenceCount(InWeakPtr.WeakReferenceCount)
 	{
@@ -1383,7 +1381,7 @@ public:
 		typename OtherType
 		UE_REQUIRES(std::is_convertible_v<OtherType*, ObjectType*>)
 	>
-	inline TWeakPtr( TWeakPtr< OtherType, Mode > const& InWeakPtr )
+	UE_NODEBUG inline TWeakPtr( TWeakPtr< OtherType, Mode > const& InWeakPtr )
 		: Object( InWeakPtr.Object )
 		, WeakReferenceCount( InWeakPtr.WeakReferenceCount )
 	{
@@ -1393,20 +1391,20 @@ public:
 		typename OtherType
 		UE_REQUIRES(std::is_convertible_v<OtherType*, ObjectType*>)
 	>
-	inline TWeakPtr( TWeakPtr< OtherType, Mode >&& InWeakPtr )
+	UE_NODEBUG inline TWeakPtr( TWeakPtr< OtherType, Mode >&& InWeakPtr )
 		: Object( InWeakPtr.Object )
 		, WeakReferenceCount( MoveTemp(InWeakPtr.WeakReferenceCount) )
 	{
 		InWeakPtr.Object = nullptr;
 	}
 
-	inline TWeakPtr( TWeakPtr const& InWeakPtr )
+	UE_NODEBUG inline TWeakPtr( TWeakPtr const& InWeakPtr )
 		: Object( InWeakPtr.Object )
 		, WeakReferenceCount( InWeakPtr.WeakReferenceCount )
 	{
 	}
 
-	inline TWeakPtr( TWeakPtr&& InWeakPtr )
+	UE_NODEBUG inline TWeakPtr( TWeakPtr&& InWeakPtr )
 		: Object( InWeakPtr.Object )
 		, WeakReferenceCount( MoveTemp(InWeakPtr.WeakReferenceCount) )
 	{
@@ -1417,7 +1415,7 @@ public:
 	 * Assignment to a nullptr pointer.  Clears this weak pointer's reference.
 	 */
 	// NOTE: The following is an Unreal extension to standard shared_ptr behavior
-	inline TWeakPtr& operator=( SharedPointerInternals::FNullTag* )
+	UE_REWRITE TWeakPtr& operator=( SharedPointerInternals::FNullTag* )
 	{
 		Reset();
 		return *this;
@@ -1428,14 +1426,14 @@ public:
 	 *
 	 * @param  InWeakPtr  The weak pointer for the object to assign
 	 */
-	inline TWeakPtr& operator=( TWeakPtr const& InWeakPtr )
+	UE_NODEBUG inline TWeakPtr& operator=( TWeakPtr const& InWeakPtr )
 	{
 		TWeakPtr Temp = InWeakPtr;
 		::Swap(Temp, *this);
 		return *this;
 	}
 
-	inline TWeakPtr& operator=( TWeakPtr&& InWeakPtr )
+	UE_NODEBUG inline TWeakPtr& operator=( TWeakPtr&& InWeakPtr )
 	{
 		if (this != &InWeakPtr)
 		{
@@ -1456,7 +1454,7 @@ public:
 		typename OtherType
 		UE_REQUIRES(std::is_convertible_v<OtherType*, ObjectType*>)
 	>
-	inline TWeakPtr& operator=( TWeakPtr<OtherType, Mode> const& InWeakPtr )
+	UE_NODEBUG inline TWeakPtr& operator=( TWeakPtr<OtherType, Mode> const& InWeakPtr )
 	{
 		Object = InWeakPtr.Pin().Get();
 		WeakReferenceCount = InWeakPtr.WeakReferenceCount;
@@ -1467,7 +1465,7 @@ public:
 		typename OtherType
 		UE_REQUIRES(std::is_convertible_v<OtherType*, ObjectType*>)
 	>
-	inline TWeakPtr& operator=( TWeakPtr<OtherType, Mode>&& InWeakPtr )
+	UE_NODEBUG inline TWeakPtr& operator=( TWeakPtr<OtherType, Mode>&& InWeakPtr )
 	{
 		Object             = InWeakPtr.Object;
 		InWeakPtr.Object   = nullptr;
@@ -1485,7 +1483,7 @@ public:
 		typename OtherType
 		UE_REQUIRES(std::is_convertible_v<OtherType*, ObjectType*>)
 	>
-	inline TWeakPtr& operator=( TSharedRef< OtherType, Mode > const& InSharedRef )
+	UE_NODEBUG inline TWeakPtr& operator=( TSharedRef< OtherType, Mode > const& InSharedRef )
 	{
 		Object = InSharedRef.Object;
 		WeakReferenceCount = InSharedRef.SharedReferenceCount;
@@ -1501,7 +1499,7 @@ public:
 		typename OtherType
 		UE_REQUIRES(std::is_convertible_v<OtherType*, ObjectType*>)
 	>
-	inline TWeakPtr& operator=( TSharedPtr< OtherType, Mode > const& InSharedPtr )
+	UE_NODEBUG inline TWeakPtr& operator=( TSharedPtr< OtherType, Mode > const& InSharedPtr )
 	{
 		Object = InSharedPtr.Object;
 		WeakReferenceCount = InSharedPtr.SharedReferenceCount;
@@ -1519,7 +1517,7 @@ public:
 	 *
 	 * @return  Shared pointer for this object (will only be valid if still referenced!)
 	 */
-	[[nodiscard]] UE_FORCEINLINE_HINT TSharedPtr< ObjectType, Mode > Pin() const&
+	[[nodiscard]] UE_REWRITE TSharedPtr< ObjectType, Mode > Pin() const&
 	{
 		return TSharedPtr< ObjectType, Mode >( *this );
 	}
@@ -1532,7 +1530,7 @@ public:
 	 *
 	 * @return  Shared pointer for this object (will only be valid if still referenced!)
 	 */
-	[[nodiscard]] UE_FORCEINLINE_HINT TSharedPtr< ObjectType, Mode > Pin() &&
+	[[nodiscard]] UE_REWRITE TSharedPtr< ObjectType, Mode > Pin() &&
 	{
 		return TSharedPtr< ObjectType, Mode >( MoveTemp( *this ) );
 	}
@@ -1542,7 +1540,7 @@ public:
 	 *
 	 * @return  True if the weak pointer is valid and a pin operator would have succeeded
 	 */
-	[[nodiscard]] UE_FORCEINLINE_HINT bool IsValid() const
+	[[nodiscard]] UE_NODEBUG inline bool IsValid() const
 	{
 		return Object != nullptr && WeakReferenceCount.IsValid();
 	}
@@ -1551,7 +1549,7 @@ public:
 	 * Resets this weak pointer, removing a weak reference to the object.  If there are no other shared
 	 * or weak references to the object, then the tracking object will be destroyed.
 	 */
-	UE_FORCEINLINE_HINT void Reset()
+	UE_REWRITE void Reset()
 	{
 		*this = TWeakPtr< ObjectType, Mode >();
 	}
@@ -1559,12 +1557,12 @@ public:
 	/**
 	 * Returns true if the object this weak pointer points to is the same as the specified object pointer.
 	 */
-	[[nodiscard]] UE_FORCEINLINE_HINT bool HasSameObject( const void* InOtherPtr ) const
+	[[nodiscard]] UE_REWRITE bool HasSameObject( const void* InOtherPtr ) const
 	{
 		return Pin().Get() == InOtherPtr;
 	}
 
-	UE_FORCEINLINE_HINT uint32 GetWeakPtrTypeHash() const
+	UE_REWRITE uint32 GetWeakPtrTypeHash() const
 	{
 		return ::PointerHash( Object );
 	}
@@ -1575,11 +1573,11 @@ public:
 	constexpr static bool bHasIntrusiveUnsetOptionalState = true;
 	using IntrusiveUnsetOptionalStateType = TWeakPtr;
 
-	explicit TWeakPtr(FIntrusiveUnsetOptionalState)
+	UE_REWRITE explicit TWeakPtr(FIntrusiveUnsetOptionalState)
 		: Object((ObjectType*)-1)
 	{
 	}
-	bool UEOpEquals(FIntrusiveUnsetOptionalState) const
+	UE_REWRITE bool UEOpEquals(FIntrusiveUnsetOptionalState) const
 	{
 		return Object == (ObjectType*)-1;
 	}
@@ -1588,7 +1586,7 @@ public:
 	///////////////////////////////////////////////
 
 	template <typename OtherType>
-	[[nodiscard]] inline bool UEOpEquals(const TWeakPtr<OtherType, Mode>& Rhs) const
+	[[nodiscard]] UE_REWRITE bool UEOpEquals(const TWeakPtr<OtherType, Mode>& Rhs) const
 	{
 		return this->Pin().Get() == Rhs.Pin().Get();
 	}
@@ -1599,7 +1597,7 @@ public:
 	}
 
 	template <typename OtherType>
-	[[nodiscard]] inline bool UEOpLessThan(const TWeakPtr<OtherType, Mode>& Rhs) const
+	[[nodiscard]] UE_REWRITE bool UEOpLessThan(const TWeakPtr<OtherType, Mode>& Rhs) const
 	{
 		return this->Pin().Get() < Rhs.Pin().Get();
 	}
@@ -1657,7 +1655,7 @@ public:
 	 *
 	 * @return	Returns this object as a shared reference
 	 */
-	[[nodiscard]] TSharedRef< ObjectType, Mode > AsShared()
+	[[nodiscard]] UE_NODEBUG TSharedRef< ObjectType, Mode > AsShared()
 	{
 		TSharedPtr< ObjectType, Mode > SharedThis( WeakThis.Pin() );
 
@@ -1684,7 +1682,7 @@ public:
 	 *
 	 * @return	Returns this object as a shared reference (const)
 	 */
-	[[nodiscard]] TSharedRef< ObjectType const, Mode > AsShared() const
+	[[nodiscard]] UE_NODEBUG TSharedRef< ObjectType const, Mode > AsShared() const
 	{
 		TSharedPtr< ObjectType const, Mode > SharedThis( WeakThis );
 
@@ -1713,7 +1711,7 @@ public:
 	 * @return	Returns the subobject as a shared reference.
 	 */
 	template< typename SubobjectType >
-	[[nodiscard]] TSharedRef< SubobjectType, Mode > AsSharedSubobject( SubobjectType* SubobjectPtr ) const
+	[[nodiscard]] UE_REWRITE TSharedRef< SubobjectType, Mode > AsSharedSubobject( SubobjectType* SubobjectPtr ) const
 	{
 		return TSharedRef< SubobjectType, Mode >( AsShared(), SubobjectPtr );
 	}
@@ -1725,7 +1723,7 @@ public:
 	 *
 	 * @return	Returns this object as a weak pointer
 	 */
-	[[nodiscard]] TWeakPtr< ObjectType, Mode > AsWeak()
+	[[nodiscard]] UE_NODEBUG TWeakPtr< ObjectType, Mode > AsWeak()
 	{
 		TWeakPtr< ObjectType, Mode > Result = WeakThis;
 
@@ -1743,7 +1741,7 @@ public:
 		// Now that we've verified the pointer is valid, we'll return it!
 		return Result;
 	}
-	[[nodiscard]] TWeakPtr< ObjectType const, Mode > AsWeak() const
+	[[nodiscard]] UE_NODEBUG TWeakPtr< ObjectType const, Mode > AsWeak() const
 	{
 		TWeakPtr< ObjectType const, Mode > Result = WeakThis;
 
@@ -1771,7 +1769,7 @@ public:
 	 * @return	Returns the subobject as a weak pointer.
 	 */
 	template <typename SubobjectType>
-	[[nodiscard]] TWeakPtr< SubobjectType, Mode > AsWeakSubobject( SubobjectType* SubobjectPtr ) const
+	[[nodiscard]] UE_REWRITE TWeakPtr< SubobjectType, Mode > AsWeakSubobject( SubobjectType* SubobjectPtr ) const
 	{
 		return AsSharedSubobject( SubobjectPtr ).ToWeakPtr();
 	}
@@ -1787,7 +1785,7 @@ protected:
 	 * @return	Returns this object as a shared pointer
 	 */
 	template< class OtherType >
-	[[nodiscard]] UE_FORCEINLINE_HINT static TSharedRef< OtherType, Mode > SharedThis( OtherType* ThisPtr )
+	[[nodiscard]] UE_REWRITE static TSharedRef< OtherType, Mode > SharedThis( OtherType* ThisPtr )
 	{
 		return StaticCastSharedRef< OtherType >( ThisPtr->AsShared() );
 	}
@@ -1801,7 +1799,7 @@ protected:
 	 * @return	Returns this object as a shared pointer (const)
 	 */
 	template< class OtherType >
-	[[nodiscard]] UE_FORCEINLINE_HINT static TSharedRef< OtherType const, Mode > SharedThis( const OtherType* ThisPtr )
+	[[nodiscard]] UE_REWRITE static TSharedRef< OtherType const, Mode > SharedThis( const OtherType* ThisPtr )
 	{
 		return StaticCastSharedRef< OtherType const >( ThisPtr->AsShared() );
 	}
@@ -1814,7 +1812,7 @@ public:		// @todo: Ideally this would be private, but template sharing problems
 	 * Note that until this function is called, calls to AsShared() will result in an empty pointer.
 	 */
 	template< class SharedPtrType, ESPMode SharedPtrMode, class OtherType >
-	inline void UpdateWeakReferenceInternal( TSharedPtr< SharedPtrType, SharedPtrMode > const* InSharedPtr, OtherType* InObject ) const
+	UE_NODEBUG inline void UpdateWeakReferenceInternal( TSharedPtr< SharedPtrType, SharedPtrMode > const* InSharedPtr, OtherType* InObject ) const
 	{
 		static_assert(SharedPtrMode == Mode, "You cannot use a TSharedPtr of one mode with a type which inherits TSharedFromThis of another mode.");
 
@@ -1830,7 +1828,7 @@ public:		// @todo: Ideally this would be private, but template sharing problems
 	 * Note that until this function is called, calls to AsShared() will result in an empty pointer.
 	 */
 	template< class SharedRefType, ESPMode SharedPtrMode, class OtherType >
-	inline void UpdateWeakReferenceInternal( TSharedRef< SharedRefType, SharedPtrMode > const* InSharedRef, OtherType* InObject ) const
+	UE_NODEBUG inline void UpdateWeakReferenceInternal( TSharedRef< SharedRefType, SharedPtrMode > const* InSharedRef, OtherType* InObject ) const
 	{
 		static_assert(SharedPtrMode == Mode, "You cannot use a TSharedPtr of one mode with a type which inherits TSharedFromThis of another mode.");
 
@@ -1847,7 +1845,7 @@ public:		// @todo: Ideally this would be private, but template sharing problems
 	 *  - Is currently within its constructor (so the shared instance isn't yet available).
 	 *  - Is currently within its destructor (so the shared instance is no longer available).
 	 */
-	[[nodiscard]] UE_FORCEINLINE_HINT bool DoesSharedInstanceExist() const
+	[[nodiscard]] UE_REWRITE bool DoesSharedInstanceExist() const
 	{
 		return WeakThis.IsValid();
 	}
@@ -1869,7 +1867,7 @@ protected:
 	TSharedFromThis( TSharedFromThis const& ) { }
 
 	/** Hidden stub assignment operator */
-	UE_FORCEINLINE_HINT TSharedFromThis& operator=( TSharedFromThis const& )
+	UE_REWRITE TSharedFromThis& operator=( TSharedFromThis const& )
 	{
 		return *this;
 	}
@@ -1891,7 +1889,7 @@ private:
  * @param  InSharedPtr  The shared pointer to cast
  */
 template< class CastToType, class CastFromType, ESPMode Mode >
-[[nodiscard]] UE_FORCEINLINE_HINT TSharedPtr< CastToType, Mode > StaticCastSharedPtr(TSharedPtr< CastFromType, Mode > const& InSharedPtr)
+[[nodiscard]] UE_REWRITE TSharedPtr< CastToType, Mode > StaticCastSharedPtr(TSharedPtr< CastFromType, Mode > const& InSharedPtr)
 {
 	return TSharedPtr< CastToType, Mode >(InSharedPtr, SharedPointerInternals::FStaticCastTag());
 }
@@ -1903,7 +1901,7 @@ template< class CastToType, class CastFromType, ESPMode Mode >
  * @param  InWeakPtr  The weak pointer to cast
  */
 template< class CastToType, class CastFromType, ESPMode Mode >
-[[nodiscard]] UE_FORCEINLINE_HINT TWeakPtr< CastToType, Mode > StaticCastWeakPtr(TWeakPtr< CastFromType, Mode > const& InWeakPtr)
+[[nodiscard]] UE_REWRITE TWeakPtr< CastToType, Mode > StaticCastWeakPtr(TWeakPtr< CastFromType, Mode > const& InWeakPtr)
 {
 	return TWeakPtr< CastToType, Mode >(InWeakPtr, SharedPointerInternals::FStaticCastTag());
 }
@@ -1915,7 +1913,7 @@ template< class CastToType, class CastFromType, ESPMode Mode >
  * @param  InSharedRef  The shared reference to cast
  */
 template< class CastToType, class CastFromType, ESPMode Mode >
-[[nodiscard]] UE_FORCEINLINE_HINT TSharedRef< CastToType, Mode > ConstCastSharedRef( TSharedRef< CastFromType, Mode > const& InSharedRef )
+[[nodiscard]] UE_REWRITE TSharedRef< CastToType, Mode > ConstCastSharedRef( TSharedRef< CastFromType, Mode > const& InSharedRef )
 {
 	return TSharedRef< CastToType, Mode >( InSharedRef, SharedPointerInternals::FConstCastTag() );
 }
@@ -1927,7 +1925,7 @@ template< class CastToType, class CastFromType, ESPMode Mode >
  * @param  InSharedPtr  The shared pointer to cast
  */
 template< class CastToType, class CastFromType, ESPMode Mode >
-[[nodiscard]] UE_FORCEINLINE_HINT TSharedPtr< CastToType, Mode > ConstCastSharedPtr( TSharedPtr< CastFromType, Mode > const& InSharedPtr )
+[[nodiscard]] UE_REWRITE TSharedPtr< CastToType, Mode > ConstCastSharedPtr( TSharedPtr< CastFromType, Mode > const& InSharedPtr )
 {
 	return TSharedPtr< CastToType, Mode >( InSharedPtr, SharedPointerInternals::FConstCastTag() );
 }
@@ -1939,7 +1937,7 @@ template< class CastToType, class CastFromType, ESPMode Mode >
  * @param  InWeakPtr  The weak pointer to cast
  */
 template< class CastToType, class CastFromType, ESPMode Mode >
-[[nodiscard]] UE_FORCEINLINE_HINT TWeakPtr< CastToType, Mode > ConstCastWeakPtr(TWeakPtr< CastFromType, Mode > const& InWeakPtr)
+[[nodiscard]] UE_REWRITE TWeakPtr< CastToType, Mode > ConstCastWeakPtr(TWeakPtr< CastFromType, Mode > const& InWeakPtr)
 {
 	return TWeakPtr< CastToType, Mode >(InWeakPtr, SharedPointerInternals::FConstCastTag());
 }
@@ -2067,7 +2065,7 @@ inline void CleanupPointerMap(TMap< TWeakPtr<KeyType>, ValueType >& PointerMap)
 * @return  Hash code value
 */
 template <typename ObjectType, ESPMode Mode>
-[[nodiscard]] uint32 GetTypeHash( const TSharedRef< ObjectType, Mode >& InSharedRef )
+[[nodiscard]] UE_REWRITE uint32 GetTypeHash( const TSharedRef< ObjectType, Mode >& InSharedRef )
 {
 	return ::PointerHash( &InSharedRef.Get() );
 }
@@ -2080,7 +2078,7 @@ template <typename ObjectType, ESPMode Mode>
 * @return  Hash code value
 */
 template <typename ObjectType, ESPMode Mode>
-[[nodiscard]] uint32 GetTypeHash( const TSharedPtr< ObjectType, Mode >& InSharedPtr )
+[[nodiscard]] UE_REWRITE uint32 GetTypeHash( const TSharedPtr< ObjectType, Mode >& InSharedPtr )
 {
 	return ::PointerHash( InSharedPtr.Get() );
 }
@@ -2093,7 +2091,7 @@ template <typename ObjectType, ESPMode Mode>
 * @return  Hash code value
 */
 template <typename ObjectType, ESPMode Mode>
-[[nodiscard]] uint32 GetTypeHash( const TWeakPtr< ObjectType, Mode >& InWeakPtr )
+[[nodiscard]] UE_REWRITE uint32 GetTypeHash( const TWeakPtr< ObjectType, Mode >& InWeakPtr )
 {
 	return InWeakPtr.GetWeakPtrTypeHash();
 }
