Added a few missing property types to UObjectHierarchyFwd.h.

diff --git a/Engine/Source/Runtime/Core/Public/UObject/UObjectHierarchyFwd.h b/Engine/Source/Runtime/Core/Public/UObject/UObjectHierarchyFwd.h
--- a/Engine/Source/Runtime/Core/Public/UObject/UObjectHierarchyFwd.h
+++ b/Engine/Source/Runtime/Core/Public/UObject/UObjectHierarchyFwd.h
@@ -49,8 +49,11 @@ class		FTextProperty;
 class		FArrayProperty;
 class		FDelegateProperty;
 class		FMulticastDelegateProperty;
+class		FMulticastSparseDelegateProperty;
 class		FMapProperty;
 class		FSetProperty;
 class		FEnumProperty;
 class		FUtf8StrProperty;
 class		FAnsiStrProperty;
+class		FOptionalProperty;
+class		FFieldPathProperty;
