Added a compiler flag & accompanying CVar to allow the work-stealing task queue to explicitly yield the thread if it fails to steal a job after a certain number of iterations.

diff --git a/Engine/Source/Runtime/Core/Private/Async/TaskGraph.cpp b/Engine/Source/Runtime/Core/Private/Async/TaskGraph.cpp
--- a/Engine/Source/Runtime/Core/Private/Async/TaskGraph.cpp
+++ b/Engine/Source/Runtime/Core/Private/Async/TaskGraph.cpp
@@ -142,6 +142,14 @@ static FAutoConsoleVariableRef CVarTaskThreadStackSize(
 	TEXT("Thread Stack Size, 0 means platform default")
 );
 
+#if UE_ASYNC_WORK_STEALING_QUEUE_USES_RETRY_LIMIT
+CORE_API int32 GWorkStealingQueueRetryLimit = -1;
+static FAutoConsoleVariableRef WorkStealingQueueRetryLimitCVar(
+	TEXT("TaskGraph.WorkStealingQueue.RetryLimit"),
+	GWorkStealingQueueRetryLimit,
+	TEXT("The maximum number of times a work-stealing attempt can see a free/taken slot before yielding"),
+	ECVF_Default);
+#endif
 
 #if CREATE_HIPRI_TASK_THREADS || CREATE_BACKGROUND_TASK_THREADS
 	static void ThreadSwitchForABTest(const TArray<FString>& Args)
diff --git a/Engine/Source/Runtime/Core/Public/Async/Fundamental/LocalQueue.h b/Engine/Source/Runtime/Core/Public/Async/Fundamental/LocalQueue.h
--- a/Engine/Source/Runtime/Core/Public/Async/Fundamental/LocalQueue.h
+++ b/Engine/Source/Runtime/Core/Public/Async/Fundamental/LocalQueue.h
@@ -17,6 +17,14 @@
 	#define LOCALQUEUEREGISTRYDEFAULTS_MAX_ITEMCOUNT 1024
 #endif
 
+#if !defined(UE_ASYNC_WORK_STEALING_QUEUE_USES_RETRY_LIMIT)
+#define UE_ASYNC_WORK_STEALING_QUEUE_USES_RETRY_LIMIT 0
+#endif
+
+#if UE_ASYNC_WORK_STEALING_QUEUE_USES_RETRY_LIMIT
+extern CORE_API int32 GWorkStealingQueueRetryLimit;
+#endif
+
 namespace LowLevelTasks
 {
 namespace LocalQueue_Impl
@@ -69,6 +77,10 @@ protected:
 	//remove an item at the tail position in LIFO order (this can be done from any thread including the one that accesses the head)
 	inline bool Steal(uintptr_t& Item)
 	{
+#if UE_ASYNC_WORK_STEALING_QUEUE_USES_RETRY_LIMIT
+		int32 Retries = GWorkStealingQueueRetryLimit;
+#endif
+
 		do
 		{
 			uint32 IdxVer = Tail.load(std::memory_order_acquire);
@@ -87,18 +99,28 @@ protected:
 				}
 				return false;
 			}
-			else if (Slot != uintptr_t(ESlotState::Taken) && ItemSlots[Idx].Value.compare_exchange_weak(Slot, uintptr_t(ESlotState::Taken), std::memory_order_acq_rel))
+			else if (Slot != uintptr_t(ESlotState::Taken))
 			{
-				if(IdxVer == Tail.load(std::memory_order_acquire))
+				if (ItemSlots[Idx].Value.compare_exchange_weak(Slot, uintptr_t(ESlotState::Taken), std::memory_order_acq_rel))
 				{
-					uint32 Prev = Tail.fetch_add(1, std::memory_order_release); (void)Prev;
-					checkSlow(Prev % NumItems == Idx);
-					ItemSlots[Idx].Value.store(uintptr_t(ESlotState::Free), std::memory_order_release);
-					Item = Slot;
-					return true;
+					if(IdxVer == Tail.load(std::memory_order_acquire))
+					{
+						uint32 Prev = Tail.fetch_add(1, std::memory_order_release); (void)Prev;
+						checkSlow(Prev % NumItems == Idx);
+						ItemSlots[Idx].Value.store(uintptr_t(ESlotState::Free), std::memory_order_release);
+						Item = Slot;
+						return true;
+					}
+					ItemSlots[Idx].Value.store(Slot, std::memory_order_release);
 				}
-				ItemSlots[Idx].Value.store(Slot, std::memory_order_release);
 			}
+#if UE_ASYNC_WORK_STEALING_QUEUE_USES_RETRY_LIMIT
+			else if (Retries-- == 0)
+			{
+				FPlatformProcess::YieldThread();
+				Retries = GWorkStealingQueueRetryLimit;
+			}
+#endif
 		} while(true);
 	}
 
