Use atomic read with PrimitiveDefinitelyUnoccludedMap. Added AtomicGet() member to TBitReference to match AtomicSet().

#Synced-CL 50642529

diff --git a/Engine/Source/Runtime/Core/Public/Containers/BitArray.h b/Engine/Source/Runtime/Core/Public/Containers/BitArray.h
--- a/Engine/Source/Runtime/Core/Public/Containers/BitArray.h
+++ b/Engine/Source/Runtime/Core/Public/Containers/BitArray.h
@@ -241,6 +241,11 @@ public:
 			}
 		}
 	}
+	[[nodiscard]] UE_FORCEINLINE_HINT bool AtomicGet() const
+	{
+		const uint32 Current = static_cast<uint32>(FPlatformAtomics::AtomicRead_Relaxed((const volatile int32*)&Data));
+		return (Current & Mask) != 0;
+	}
 	UE_FORCEINLINE_HINT FBitReference& operator=(const FBitReference& Copy)
 	{
 		// As this is emulating a reference, assignment should not rebind,
@@ -270,6 +275,12 @@ public:
 		 return (Data & Mask) != 0;
 	}
 
+	[[nodiscard]] UE_FORCEINLINE_HINT bool AtomicGet() const
+	{
+		const uint32 Current = static_cast<uint32>(FPlatformAtomics::AtomicRead_Relaxed((const volatile int32*)&Data));
+		return (Current & Mask) != 0;
+	}
+
 private:
 	const uint32& Data;
 	uint32 Mask;
diff --git a/Engine/Source/Runtime/Renderer/Private/SceneVisibility.cpp b/Engine/Source/Runtime/Renderer/Private/SceneVisibility.cpp
--- a/Engine/Source/Runtime/Renderer/Private/SceneVisibility.cpp
+++ b/Engine/Source/Runtime/Renderer/Private/SceneVisibility.cpp
@@ -2160,7 +2160,7 @@ void FRelevancePacket::ComputeRelevance(FDynamicPrimitiveIndexList& DynamicPrimi
 		PrimitiveSceneInfo->LastRenderTime = CurrentWorldTime;
 
 		// Update the last component render time only if we know for certain the primitive is un-occluded.
-		if ((View.PrimitiveDefinitelyUnoccludedMap[BitIndex] || View.Family->EngineShowFlags.Wireframe) && !PrimitiveSceneProxy->IsAlwaysVisible())
+		if ((View.PrimitiveDefinitelyUnoccludedMap[BitIndex].AtomicGet() || View.Family->EngineShowFlags.Wireframe) && !PrimitiveSceneProxy->IsAlwaysVisible())
 		{
 			const bool bUpdateLastRenderTimeOnScreen = true;
 			PrimitiveSceneInfo->UpdateComponentLastRenderTime(CurrentWorldTime, bUpdateLastRenderTimeOnScreen);
