Log the config section and key in GetString and GetText when it is about to call SafeReload. Also add delegates that broadcast in SafeUnload and SafeReload for additional instrumentation of related hitches.

#tests Windows client
#rnx

#Synced-CL 50645574

diff --git a/Engine/Source/Runtime/Core/Private/Misc/ConfigCacheIni.cpp b/Engine/Source/Runtime/Core/Private/Misc/ConfigCacheIni.cpp
--- a/Engine/Source/Runtime/Core/Private/Misc/ConfigCacheIni.cpp
+++ b/Engine/Source/Runtime/Core/Private/Misc/ConfigCacheIni.cpp
@@ -3856,6 +3856,8 @@ void FConfigBranch::SafeUnload()
 			Node->Empty();
 		}
 	}
+
+	FCoreDelegates::TSOnConfigBranchUnloaded().Broadcast(*IniName.ToString());
 }
 
 void FConfigBranch::SafeReload()
@@ -3909,7 +3911,11 @@ void FConfigBranch::SafeReload()
 		AddDynamicLayersToHierarchy(ReloadInfos, nullptr, nullptr, nullptr, bForceFullDynamicLayerUpdate);
 	}
 
-	UE_LOG(LogConfig, Log, TEXT("Branch '%s' had been unloaded. Reloading on-demand took %.2fms"), *IniName.ToString(), (FPlatformTime::Seconds() - StartTime) * 1000.0f);
+	const FString IniNameStr = IniName.ToString();
+	const float ReloadTimeMs = (float)(FPlatformTime::Seconds() - StartTime) * 1000.0f;
+	UE_LOG(LogConfig, Log, TEXT("Branch '%s' had been unloaded. Reloading on-demand took %.2fms"), *IniNameStr, ReloadTimeMs);
+
+	FCoreDelegates::TSOnConfigBranchReloaded().Broadcast(*IniNameStr, ReloadTimeMs);
 }
 
 void FConfigBranch::ReapplyLayers()
@@ -4172,7 +4178,7 @@ void FConfigCacheIni::Tick(float DeltaSeconds)
 		{
 			if (Now - Branch->InactiveTimer > GTimeToUnloadConfig)
 			{
-				UE_LOG(LogConfig, Verbose, TEXT("Unloading %s due to inactivity"), *Branch->IniPath);
+				UE_LOG(LogConfig, Log, TEXT("Unloading %s due to inactivity"), *Branch->IniPath);
 				
 				Branch->SafeUnload();
 				Branch->InactiveTimer = 0;
@@ -4771,6 +4777,7 @@ bool FConfigCacheIni::GetString( const TCHAR* Section, const TCHAR* Key, FString
 	
 	if (Branch->bIsSafeUnloaded)
 	{
+		UE_LOG(LogConfig, Log, TEXT("About to reload due to getting %s: %s"), Section, Key);
 		Branch->SafeReload();
 	}
 #else
@@ -4822,6 +4829,7 @@ bool FConfigCacheIni::GetText( const TCHAR* Section, const TCHAR* Key, FText& Va
 	
 	if (Branch->bIsSafeUnloaded)
 	{
+		UE_LOG(LogConfig, Log, TEXT("About to reload due to getting %s: %s"), Section, Key);
 		Branch->SafeReload();
 	}
 #else
diff --git a/Engine/Source/Runtime/Core/Private/Misc/CoreDelegates.cpp b/Engine/Source/Runtime/Core/Private/Misc/CoreDelegates.cpp
--- a/Engine/Source/Runtime/Core/Private/Misc/CoreDelegates.cpp
+++ b/Engine/Source/Runtime/Core/Private/Misc/CoreDelegates.cpp
@@ -244,6 +244,18 @@ TTSMulticastDelegate<void(const FString&, const TSet<FString>&)>& FCoreDelegates
 	return Singleton;
 }
 
+TTSMulticastDelegate<void(const TCHAR*)>& FCoreDelegates::TSOnConfigBranchUnloaded()
+{
+	static TTSMulticastDelegate<void(const TCHAR*)> Singleton;
+	return Singleton;
+}
+
+TTSMulticastDelegate<void(const TCHAR*, float)>& FCoreDelegates::TSOnConfigBranchReloaded()
+{
+	static TTSMulticastDelegate<void(const TCHAR*, float)> Singleton;
+	return Singleton;
+}
+
 TMulticastDelegate<void(const TCHAR*, const TCHAR*, uint32, bool)> FCoreDelegates::OnApplyCVarFromIni;
 TMulticastDelegate<void(uint32, uint32)> FCoreDelegates::OnSystemResolutionChanged;
 
diff --git a/Engine/Source/Runtime/Core/Public/Misc/CoreDelegates.h b/Engine/Source/Runtime/Core/Public/Misc/CoreDelegates.h
--- a/Engine/Source/Runtime/Core/Public/Misc/CoreDelegates.h
+++ b/Engine/Source/Runtime/Core/Public/Misc/CoreDelegates.h
@@ -317,6 +317,8 @@ public:
 	static CORE_API TTSMulticastDelegate<void(const TCHAR* IniFilename, const TCHAR* SectionName)>& TSOnConfigSectionRead();
 	static CORE_API TTSMulticastDelegate<void(const TCHAR* IniFilename, const TCHAR* SectionName)>& TSOnConfigSectionNameRead();
 	static CORE_API TTSMulticastDelegate<void(const FString& IniFilename, const TSet<FString>& SectionNames)>& TSOnConfigSectionsChanged();
+	static CORE_API TTSMulticastDelegate<void(const TCHAR* IniFilename)>& TSOnConfigBranchUnloaded();
+	static CORE_API TTSMulticastDelegate<void(const TCHAR* IniFilename, float ReloadTimeMs)>& TSOnConfigBranchReloaded();
 
 	static CORE_API TMulticastDelegate<void(const TCHAR* SectionName, const TCHAR* IniFilename, uint32 SetBy, bool bAllowCheating)> OnApplyCVarFromIni;
 
