Fix regression where the FConfigBranch inactive timer would never be reset through FConfigCacheIni::GetString or GetText, keeping the branch loaded forever unless other functions are called to set the timer to the current time.

Fix inconsistency where InactiveTimer would never be reset to 0 when a branch is unloaded via FConfigCache::SafeUnloadBranch, such as with the "config unload" or "config unloadall" console commands.

#tests Windows client
#rnx
#Synced-CL 50645574

diff --git a/Engine/Source/Runtime/Core/Private/Misc/ConfigCacheIni.cpp b/Engine/Source/Runtime/Core/Private/Misc/ConfigCacheIni.cpp
--- a/Engine/Source/Runtime/Core/Private/Misc/ConfigCacheIni.cpp
+++ b/Engine/Source/Runtime/Core/Private/Misc/ConfigCacheIni.cpp
@@ -3831,6 +3831,7 @@ void FConfigBranch::SafeUnload()
 #endif
 	
 	bIsSafeUnloaded = true;
+	InactiveTimer = 0;
 
 #if UE_THREAD_SAFETY_USE_WRITE_LOCKS
 	// if we are not using full Write Locks, then we we no longer unload the final combined data
@@ -3918,6 +3919,15 @@ void FConfigBranch::SafeReload()
 	FCoreDelegates::TSOnConfigBranchReloaded().Broadcast(*IniNameStr, ReloadTimeMs);
 }
 
+void FConfigBranch::UpdateInactiveTimer()
+{
+	if (InactiveTimer >= 0 && GTimeToUnloadConfig > 0)
+	{
+		InactiveTimer = FPlatformTime::Seconds();
+		UE_LOG(LogConfig, Verbose, TEXT("Resetting InactiveTimer for %s"), *IniName.ToString());
+	}
+}
+
 void FConfigBranch::ReapplyLayers()
 {
 #if UE_THREAD_SAFETY_PROTECT_GETSTRING
@@ -4181,7 +4191,6 @@ void FConfigCacheIni::Tick(float DeltaSeconds)
 				UE_LOG(LogConfig, Log, TEXT("Unloading %s due to inactivity"), *Branch->IniPath);
 				
 				Branch->SafeUnload();
-				Branch->InactiveTimer = 0;
 			}
 		}
 	}
@@ -4221,6 +4230,12 @@ FConfigBranch* FConfigCacheIni::FindBranchWithNoReload(FName BaseIniName, const
 		Branch = nullptr;
 	}
 
+	if (Branch)
+	{
+		// track that this branch is being used, so re-set the time
+		Branch->UpdateInactiveTimer();
+	}
+
 	return Branch;
 }
 
@@ -4230,17 +4245,13 @@ FConfigBranch* FConfigCacheIni::FindBranch(FName BaseIniName, const FString& Fil
 
 	if (Branch)
 	{
-		// track that this branch is being used, so re-set the time
-		if (Branch->InactiveTimer >= 0 && GTimeToUnloadConfig > 0)
-		{
-			Branch->InactiveTimer = FPlatformTime::Seconds();
-			UE_LOG(LogConfig, Verbose, TEXT("Resetting InactiveTimer for %s"), *Branch->IniName.ToString());
-		}
-		
 		if (Branch->bIsSafeUnloaded)
 		{
 			Branch->SafeReload();
 		}
+
+		// track that this branch is being used, so re-set the time
+		Branch->UpdateInactiveTimer();
 	}
 
 	return Branch;
@@ -4302,11 +4313,8 @@ FConfigFile* FConfigCacheIni::FindConfigFile( const FString& Filename )
 		}
 
 		// track that this branch is being used, so re-set the time
-		if (Result && Result->InactiveTimer >= 0 && GTimeToUnloadConfig > 0)
-		{
-			Result->InactiveTimer = FPlatformTime::Seconds();
-			UE_LOG(LogConfig, VeryVerbose, TEXT("Resetting InactiveTimer for %s"), *Result->IniName.ToString());
-		}
+		Result->UpdateInactiveTimer();
+
 		return &Result->InMemoryFile;
 	}
 
@@ -4387,12 +4395,10 @@ FConfigFile* FConfigCacheIni::FindConfigFileWithBaseName(FName BaseName)
 	{
 			Result->SafeReload();
 		}
+
 		// track that this branch is being used, so re-set the time
-		if (Result && Result->InactiveTimer >= 0 && GTimeToUnloadConfig > 0)
-		{
-			Result->InactiveTimer = FPlatformTime::Seconds();
-			UE_LOG(LogConfig, Verbose, TEXT("Resetting InactiveTimer for %s"), *Result->IniName.ToString());
-		}
+		Result->UpdateInactiveTimer();
+
 		return &Result->InMemoryFile;
 	}
 	return nullptr;
diff --git a/Engine/Source/Runtime/Core/Public/Misc/ConfigCacheIni.h b/Engine/Source/Runtime/Core/Public/Misc/ConfigCacheIni.h
--- a/Engine/Source/Runtime/Core/Public/Misc/ConfigCacheIni.h
+++ b/Engine/Source/Runtime/Core/Public/Misc/ConfigCacheIni.h
@@ -1197,6 +1197,11 @@ public:
 	CORE_API void SafeUnload();
 	CORE_API void SafeReload();
 
+	/**
+	 * Updates the inactive timer to the current time
+	 */
+	CORE_API void UpdateInactiveTimer();
+
 	/** Recalculate the usable values from all static and dynamic layers. Useful if you update a layer and need to see results in 
 	 * the InMemoryFile (ie what GConfig->GetXXX returns) 
 	 */
