Fixed dynamic delegate payloads when all parameters are bound and some of them get captured as different types.

#Synced-CL 50651673

diff --git a/Engine/Source/Runtime/Core/Public/UObject/ScriptDelegates.h b/Engine/Source/Runtime/Core/Public/UObject/ScriptDelegates.h
--- a/Engine/Source/Runtime/Core/Public/UObject/ScriptDelegates.h
+++ b/Engine/Source/Runtime/Core/Public/UObject/ScriptDelegates.h
@@ -32,6 +32,9 @@
 	#include "Concepts/Arithmetic.h"
 	#include "Concepts/UObject.h"
 	#include "Templates/Tuple.h"
+	#include "Traits/FunctionArity.h"
+	#include "UObject/NativeTypeToPropertyType.h"
+	#include "UObject/UObjectHierarchyFwd.h"
 	#include "UObject/WeakObjectPtrTemplatesFwd.h"
 #endif
 
@@ -111,6 +114,31 @@ namespace UE::Core::Private
 	template <typename WeakPtrType, typename T>
 	using TDynamicDelegateParamPayloadCaptureType_T = typename TDynamicDelegateParamPayloadCaptureType<WeakPtrType, T>::Type;
 
+	template <typename FunctionType, typename TupleType>
+	struct TDynamicDelegateHasExactPayloadType
+	{
+		static constexpr inline bool Value = false;
+	};
+
+	template <typename... ParamTypes, typename... VarTypes>
+	struct TDynamicDelegateHasExactPayloadType<void(ParamTypes...), TTuple<VarTypes...>>
+	{
+		static constexpr inline bool Value = (std::is_same_v<std::decay_t<ParamTypes>, VarTypes> && ...);
+	};
+
+	template <typename FunctionType, typename TupleType>
+	constexpr inline bool DynamicDelegatePayloadIsLayoutCompatible()
+	{
+		if constexpr (TFunctionArity_V<FunctionType> != TTupleArity<TupleType>::Value)
+		{
+			return false;
+		}
+		else
+		{
+			return TDynamicDelegateHasExactPayloadType<FunctionType, TupleType>::Value;
+		}
+	}
+
 	// An interface to a dynamic delegate payload.  Type-erases payload operations.
 	template <typename WeakPtrType>
 	struct IDynamicDelegatePayload
@@ -130,6 +158,8 @@ namespace UE::Core::Private
 	template <typename WeakPtrType, typename... ParamTypes, typename... VarTypes>
 	struct TDynamicDelegatePayload<WeakPtrType, void(ParamTypes...), VarTypes...> : IDynamicDelegatePayload<WeakPtrType>
 	{
+		using VarsTupleType = TTuple<VarTypes...>;
+
 		static_assert(sizeof...(VarTypes) != 0, "Allocating an object to hold an empty payload is wasteful");
 
 		template <typename... ArgTypes>
@@ -145,7 +175,7 @@ namespace UE::Core::Private
 			// be passed to the ProcessEvent call.
 
 			// If the payload replaces every argument, just pass the tuple itself as it should be the same layout.
-			if constexpr (sizeof...(ParamTypes) == sizeof...(VarTypes))
+			if constexpr (DynamicDelegatePayloadIsLayoutCompatible<void(ParamTypes...), VarsTupleType>())
 			{
 				ObjectPtr.Get()->ProcessEvent(Function, (void*)&Payload);
 			}
@@ -154,7 +184,6 @@ namespace UE::Core::Private
 			{
 				using FullTupleType   = TTuple<std::decay_t<ParamTypes>...>;
 				using ParamsTupleType = TTuplePopBack_T<TTuple<std::decay_t<ParamTypes>...>, sizeof...(VarTypes)>;
-				using VarsTupleType   = TTuple<std::decay_t<VarTypes>...>;
 
 				((ParamsTupleType*)Parameters)->ApplyBefore(
 					[&Payload = Payload, ObjectPtr, Function](auto&&... Params)
@@ -174,7 +203,7 @@ namespace UE::Core::Private
 
 	private:
 		// Mutable because ProcessEvent takes a non-const param pointer
-		mutable TTuple<VarTypes...> Payload;
+		mutable VarsTupleType Payload;
 	};
 
 	template <typename WeakPtrType, typename BoundFuncType, typename... VarTypes>
