Simplify SRW code in FWindowsPlatformMisc::LocalPrint

#Synced-CL 50763424
#Synced-CL 50763424

diff --git a/Engine/Source/Runtime/Core/Private/Windows/WindowsPlatformMisc.cpp b/Engine/Source/Runtime/Core/Private/Windows/WindowsPlatformMisc.cpp
--- a/Engine/Source/Runtime/Core/Private/Windows/WindowsPlatformMisc.cpp
+++ b/Engine/Source/Runtime/Core/Private/Windows/WindowsPlatformMisc.cpp
@@ -1404,8 +1404,7 @@ void FWindowsPlatformMisc::LocalPrint(const TCHAR* Message)
 {
 #if USE_DEBUG_LOGGING
 	// Use SRWLOCK directly until we have a trivially destructible FSharedMutex.
-	static SRWLOCK PrivateLock;
-	static PSRWLOCK Lock = (InitializeSRWLock(&PrivateLock), &PrivateLock);
+	static SRWLOCK Lock = SRWLOCK_INIT;
 
 	// OutputDebugString has no documented maximum length but the OUTPUT_DEBUG_STRING_INFO
 	// type that receives debug strings in debuggers has a WORD count of bytes including
@@ -1414,20 +1413,20 @@ void FWindowsPlatformMisc::LocalPrint(const TCHAR* Message)
 	constexpr int32 MaxOutputDebugStringLength = WORD_MAX / sizeof(TCHAR) - 1;
 	if (const int32 Len = FCString::Strlen(Message); LIKELY(Len <= MaxOutputDebugStringLength))
 	{
-		AcquireSRWLockShared(Lock);
+		AcquireSRWLockShared(&Lock);
 		OutputDebugString(Message);
-		ReleaseSRWLockShared(Lock);
+		ReleaseSRWLockShared(&Lock);
 	}
 	else
 	{
-		AcquireSRWLockExclusive(Lock);
+		AcquireSRWLockExclusive(&Lock);
 		FStringView MessageView(Message, Len);
 		for (FString Part; !MessageView.IsEmpty(); MessageView.RightChopInline(MaxOutputDebugStringLength))
 		{
 			Part = MessageView.Left(MaxOutputDebugStringLength);
 			OutputDebugString(*Part);
 		}
-		ReleaseSRWLockExclusive(Lock);
+		ReleaseSRWLockExclusive(&Lock);
 	}
 #endif
 }
