Fixed disabling the game preview failing to reload the native localization data for plugins

FLocalizationTargetDescriptor::ShouldLoadLocalizationTarget relies on the global bIsGameLocalizationPreviewEnabled value being true in order to load any game localization data (native or localized) from plugins, so we have to keep it set to true until after the localization data has been reloaded when disabling the game preview

#Synced-CL 50687180

diff --git a/Engine/Source/Runtime/Core/Private/Internationalization/TextLocalizationManager.cpp b/Engine/Source/Runtime/Core/Private/Internationalization/TextLocalizationManager.cpp
--- a/Engine/Source/Runtime/Core/Private/Internationalization/TextLocalizationManager.cpp
+++ b/Engine/Source/Runtime/Core/Private/Internationalization/TextLocalizationManager.cpp
@@ -1492,11 +1492,15 @@ void FTextLocalizationManager::LoadLocalizationResourcesForPrioritizedCultures_S
 	}
 }
 
-void FTextLocalizationManager::LoadLocalizationResourcesForPrioritizedCultures_Async(TArrayView<const FString> PrioritizedCultureNames, const ELocalizationLoadFlags LocLoadFlags)
+void FTextLocalizationManager::LoadLocalizationResourcesForPrioritizedCultures_Async(TArrayView<const FString> PrioritizedCultureNames, const ELocalizationLoadFlags LocLoadFlags, TUniqueFunction<void()>&& Notification)
 {
-	QueueAsyncTask([AvailableTextSources = LocalizedTextSources, PrioritizedCultureNames = TArray<FString>(PrioritizedCultureNames), LocLoadFlags]()
+	QueueAsyncTask([AvailableTextSources = LocalizedTextSources, PrioritizedCultureNames = TArray<FString>(PrioritizedCultureNames), LocLoadFlags, Notification = MoveTemp(Notification)]()
 	{
 		FTextLocalizationManager::Get().LoadLocalizationResourcesForPrioritizedCultures_Sync(AvailableTextSources, PrioritizedCultureNames, LocLoadFlags);
+		if (Notification)
+		{
+			Notification();
+		}
 	});
 }
 
@@ -1865,11 +1869,14 @@ void FTextLocalizationManager::EnableGameLocalizationPreview(const FString& Cult
 	}
 
 	const FString PreviewCulture = CultureName.IsEmpty() ? NativeGameCulture : CultureName;
-	bIsGameLocalizationPreviewEnabled = PreviewCulture != NativeGameCulture;
-	bIsLocalizationLocked = IsLocalizationLockedByConfig() || bIsGameLocalizationPreviewEnabled;
+	const bool bWillGameLocalizationPreviewBeEnabled = PreviewCulture != NativeGameCulture;
+	// We always need to enable this, even when disabling, so that the native language data reloads correctly for plugins (see ShouldForceLoadGameLocalization in FLocalizationTargetDescriptor::ShouldLoadLocalizationTarget)
+	// The correct value is set once the loc data has loaded, via the async notification below
+	bIsGameLocalizationPreviewEnabled = true;
+	bIsLocalizationLocked = IsLocalizationLockedByConfig() || bWillGameLocalizationPreviewBeEnabled;
 
 	TArray<FString> PrioritizedCultureNames;
-	if (bIsGameLocalizationPreviewEnabled)
+	if (bWillGameLocalizationPreviewBeEnabled)
 	{
 		PrioritizedCultureNames = FInternationalization::Get().GetPrioritizedCultureNames(PreviewCulture);
 	}
@@ -1879,9 +1886,13 @@ void FTextLocalizationManager::EnableGameLocalizationPreview(const FString& Cult
 	}
 
 	ELocalizationLoadFlags LocLoadFlags = ELocalizationLoadFlags::Game | ELocalizationLoadFlags::Additional;
-	LocLoadFlags |= (bIsGameLocalizationPreviewEnabled ? ELocalizationLoadFlags::Native : ELocalizationLoadFlags::None);
+	LocLoadFlags |= (bWillGameLocalizationPreviewBeEnabled ? ELocalizationLoadFlags::Native : ELocalizationLoadFlags::None);
 
-	LoadLocalizationResourcesForPrioritizedCultures_Async(PrioritizedCultureNames, LocLoadFlags);
+	LoadLocalizationResourcesForPrioritizedCultures_Async(PrioritizedCultureNames, LocLoadFlags, 
+		[bWillGameLocalizationPreviewBeEnabled]()
+		{
+			FTextLocalizationManager::Get().bIsGameLocalizationPreviewEnabled = bWillGameLocalizationPreviewBeEnabled;
+		});
 }
 
 void FTextLocalizationManager::DisableGameLocalizationPreview()
diff --git a/Engine/Source/Runtime/Core/Public/Internationalization/TextLocalizationManager.h b/Engine/Source/Runtime/Core/Public/Internationalization/TextLocalizationManager.h
--- a/Engine/Source/Runtime/Core/Public/Internationalization/TextLocalizationManager.h
+++ b/Engine/Source/Runtime/Core/Public/Internationalization/TextLocalizationManager.h
@@ -365,7 +365,7 @@ private:
 
 	/** Loads localization resources for the specified prioritized cultures, optionally loading localization resources that are editor-specific or game-specific. */
 	void LoadLocalizationResourcesForPrioritizedCultures_Sync(TArrayView<const TSharedPtr<ILocalizedTextSource>> AvailableTextSources, TArrayView<const FString> PrioritizedCultureNames, const ELocalizationLoadFlags LocLoadFlags);
-	void LoadLocalizationResourcesForPrioritizedCultures_Async(TArrayView<const FString> PrioritizedCultureNames, const ELocalizationLoadFlags LocLoadFlags);
+	void LoadLocalizationResourcesForPrioritizedCultures_Async(TArrayView<const FString> PrioritizedCultureNames, const ELocalizationLoadFlags LocLoadFlags, TUniqueFunction<void()>&& Notification);
 
 	/** Loads the specified localization targets for the specified prioritized cultures, optionally loading localization resources that are editor-specific or game-specific. */
 	void LoadLocalizationTargetsForPrioritizedCultures_Sync(TArrayView<const TSharedPtr<ILocalizedTextSource>> AvailableTextSources, TArrayView<const FString> LocalizationTargetPaths, TArrayView<const FString> PrioritizedCultureNames, const ELocalizationLoadFlags LocLoadFlags);
