[IO] Include ChunkId and its Type in "I/O Error" messages.

A crash with fatal error "I/O Error '(Not Found)'" is not much to go on from, so we add the ChunkId and its Type to the message.
This changes the error message to something like "I/O Error '(Not Found)' (IoChunkId='0b93063ba09cabead3387509', Type=Shader)".

#rnx

#Synced-CL 50781510

diff --git a/Engine/Source/Runtime/Core/Private/IO/IoDispatcher.cpp b/Engine/Source/Runtime/Core/Private/IO/IoDispatcher.cpp
--- a/Engine/Source/Runtime/Core/Private/IO/IoDispatcher.cpp
+++ b/Engine/Source/Runtime/Core/Private/IO/IoDispatcher.cpp
@@ -9,6 +9,7 @@
 #include "HAL/Runnable.h"
 #include "HAL/RunnableThread.h"
 #include "HAL/Event.h"
+#include "IO/IoChunkId.h"
 #include "IO/IoDispatcherBackend.h"
 #include "IO/IoDispatcherConfig.h"
 #include "IO/IoDispatcherChunkDecoder.h"
@@ -1449,7 +1450,14 @@ FIoRequest::GetResultOrDie() const
 	const FIoBuffer* Result = GetResult();
 	if (!Result)
 	{
-		UE_LOG(LogIoDispatcher, Fatal, TEXT("I/O Error '%s'"), *Status().ToString());
+		if (Impl)
+		{
+			UE_LOG(LogIoDispatcher, Fatal, TEXT("I/O Error '%s' (IoChunkId='%s', Type=%s)"), *Status().ToString(), *LexToString(Impl->ChunkId), *LexToString(Impl->ChunkId.GetChunkType()));
+		}
+		else
+		{
+			UE_LOG(LogIoDispatcher, Fatal, TEXT("I/O Error '%s'"), *Status().ToString());
+		}
 	}
 	return *Result;
 }
