Circular buffer: reduced checked capacity limit that was too high and would result in a container of size 1 if attempted higher than int32max, or will hit TArray's own limits that are even lower.

#Synced-CL 50812809

diff --git a/Engine/Source/Runtime/Core/Public/Containers/CircularBuffer.h b/Engine/Source/Runtime/Core/Public/Containers/CircularBuffer.h
--- a/Engine/Source/Runtime/Core/Public/Containers/CircularBuffer.h
+++ b/Engine/Source/Runtime/Core/Public/Containers/CircularBuffer.h
@@ -26,7 +26,7 @@ public:
 	 */
 	[[nodiscard]] explicit TCircularBuffer(uint32 Capacity)
 	{
-		checkSlow(Capacity <= 0xffffffffU);
+		checkSlow(Capacity <= 0x40000000U);	// Limited by powers-of-2 counts that can fit into a TArray
 
 		Elements.AddZeroed(FMath::RoundUpToPowerOfTwo(Capacity));
 		IndexMask = Elements.Num() - 1;
@@ -40,7 +40,7 @@ public:
 	 */
 	[[nodiscard]] TCircularBuffer(uint32 Capacity, const ElementType& InitialValue)
 	{
-		checkSlow(Capacity <= 0xffffffffU);
+		checkSlow(Capacity <= 0x40000000U); // Limited by powers-of-2 counts that can fit into a TArray
 
 		Elements.Init(InitialValue, FMath::RoundUpToPowerOfTwo(Capacity));
 		IndexMask = Elements.Num() - 1;
