Tidied up FWindowsPlatformFile::GetFilenameOnDisk with RAII.

#Synced-CL 50651673

diff --git a/Engine/Source/Runtime/Core/Private/Windows/WindowsPlatformFile.cpp b/Engine/Source/Runtime/Core/Private/Windows/WindowsPlatformFile.cpp
--- a/Engine/Source/Runtime/Core/Private/Windows/WindowsPlatformFile.cpp
+++ b/Engine/Source/Runtime/Core/Private/Windows/WindowsPlatformFile.cpp
@@ -1281,6 +1281,22 @@ FDateTime FWindowsPlatformFile::GetAccessTimeStamp(const TCHAR* Filename)
 
 FString FWindowsPlatformFile::GetFilenameOnDisk(const TCHAR* Filename)
 {
+	auto ConvertToUnrealPathInline = [](FString& Path)
+	{
+		// Remove the Windows device path prefix.
+		if (Path.StartsWith(TEXT("\\\\?\\UNC\\"), ESearchCase::CaseSensitive))
+		{
+			Path.RemoveAt(2, 6); // remove ?\UNC\ to convert \\?\UNC\Path\... to \\Path\...
+		}
+		else
+		{
+			Path.RemoveFromStart(TEXT("\\\\?\\"), ESearchCase::CaseSensitive);
+		}
+
+		// Convert the result back into an UnrealPath (\\ -> /)
+		Path.ReplaceCharInline(TEXT('\\'), TEXT('/'), ESearchCase::CaseSensitive);
+	};
+
 	FString NormalizedFileName = *FNormalizedFilename(Filename);
 	TRACE_PLATFORMFILE_BEGIN_OPEN(Filename);
 	HANDLE hFile = CreateFile(*NormalizedFileName,
@@ -1289,26 +1305,15 @@ FString FWindowsPlatformFile::GetFilenameOnDisk(const TCHAR* Filename)
 		FILE_FLAG_BACKUP_SEMANTICS, // Required for opening a directory, which Filename might be
 		NULL);
 	// If the file exists on disk, read the capitalization from the path on disk, otherwise just return the (normalized) input filename
-	if (hFile != INVALID_HANDLE_VALUE)
+	if (hFile == INVALID_HANDLE_VALUE)
+	{
+		TRACE_PLATFORMFILE_FAIL_OPEN(Filename);
+		ConvertToUnrealPathInline(NormalizedFileName);
+		return NormalizedFileName;
+	}
+	TRACE_PLATFORMFILE_END_OPEN(hFile);
+	ON_SCOPE_EXIT
 	{
-		TRACE_PLATFORMFILE_END_OPEN(hFile);
-		for (uint32 Length = NormalizedFileName.Len() + 10;;)
-		{
-			TArray<TCHAR, FString::AllocatorType>& CharArray = NormalizedFileName.GetCharArray();
-			CharArray.SetNum(Length);
-
-			Length = GetFinalPathNameByHandle(hFile, CharArray.GetData(), CharArray.Num(), FILE_NAME_NORMALIZED);
-			if (Length == 0)
-			{
-				NormalizedFileName = *FNormalizedFilename(Filename);
-				break;
-			}
-			if (Length < (uint32)CharArray.Num())
-			{
-				CharArray.SetNum(Length + 1);
-				break;
-			}
-		}
 		TRACE_PLATFORMFILE_BEGIN_CLOSE(hFile);
 #if PLATFORMFILETRACE_ENABLED
 		// MSVC static analysis has a rule that reports the argument to CloseHandle as uninitialized memory after the call to CloseHandle, so we have to save it ahead of time
@@ -1327,25 +1332,29 @@ FString FWindowsPlatformFile::GetFilenameOnDisk(const TCHAR* Filename)
 #else
 		(void)CloseResult;
 #endif
-	}
-	else
-	{
-		TRACE_PLATFORMFILE_FAIL_OPEN(Filename);
-	}
+	};
 
-	// Remove the Windows device path prefix.
-	if (NormalizedFileName.StartsWith(TEXT("\\\\?\\UNC\\"), ESearchCase::CaseSensitive))
-	{
-		NormalizedFileName.RemoveAt(2, 6); // remove ?\UNC\ to convert \\?\UNC\Path\... to \\Path\...
-	}
-	else
+	for (uint32 Length = NormalizedFileName.Len() + 10;;)
 	{
-		NormalizedFileName.RemoveFromStart(TEXT("\\\\?\\"), ESearchCase::CaseSensitive);
-	}
+		TArray<TCHAR, FString::AllocatorType>& CharArray = NormalizedFileName.GetCharArray();
+		CharArray.SetNum(Length);
 
-	// Convert the result back into an UnrealPath (\\ -> /)
-	NormalizedFileName.ReplaceCharInline(TEXT('\\'), TEXT('/'), ESearchCase::CaseSensitive);
+		Length = GetFinalPathNameByHandle(hFile, CharArray.GetData(), CharArray.Num(), FILE_NAME_NORMALIZED);
+		if (Length == 0)
+		{
+			// Error - return normalized filename.
+			NormalizedFileName = *FNormalizedFilename(Filename);
+			ConvertToUnrealPathInline(NormalizedFileName);
+			return NormalizedFileName;
+		}
+		if (Length < (uint32)CharArray.Num())
+		{
+			CharArray.SetNum(Length + 1);
+			break;
+		}
+	}
 
+	ConvertToUnrealPathInline(NormalizedFileName);
 	return NormalizedFileName;
 }
 
