Change FConfigFileHierarchy to derive from TSparseMap instead of TMap.

#Synced-CL 50723010

diff --git a/Engine/Source/Runtime/Core/Public/Misc/ConfigTypes.h b/Engine/Source/Runtime/Core/Public/Misc/ConfigTypes.h
--- a/Engine/Source/Runtime/Core/Public/Misc/ConfigTypes.h
+++ b/Engine/Source/Runtime/Core/Public/Misc/ConfigTypes.h
@@ -10,7 +10,7 @@
 #include "CoreGlobals.h"
 
 
-class FConfigFileHierarchy : public TMap<int32, FUtf8String>
+class FConfigFileHierarchy : public TSparseMap<int32, FUtf8String>
 {
 public:
 	FConfigFileHierarchy();
