Enables callstack traces on Windows ARM64 and ARM64EC.

ARM64 always uses the fallback path that calls RtlCaptureStackBackTrace. Using unwind tables would require a new implementation, since they are represented differently than x64.

The fallback path broke in [REDACTED_HANDLE] when the function signature of FBacktracer::GetBacktraceId changed. This restores some of the existing code for calling this function. Additionally, it adds a new FrameSkip parameter to prevent everything from showing up as CallstackTrace_GetCurrentId_Private.

#rnx

#Synced-CL 50800567

diff --git a/Engine/Source/Runtime/Core/Private/ProfilingDebugging/Microsoft/WindowsCallstackTrace.cpp b/Engine/Source/Runtime/Core/Private/ProfilingDebugging/Microsoft/WindowsCallstackTrace.cpp
--- a/Engine/Source/Runtime/Core/Private/ProfilingDebugging/Microsoft/WindowsCallstackTrace.cpp
+++ b/Engine/Source/Runtime/Core/Private/ProfilingDebugging/Microsoft/WindowsCallstackTrace.cpp
@@ -48,6 +48,11 @@ static bool GModulesAreInitialized = false;
 	#define UE_CALLSTACK_TRACE_USE_UNWIND_TABLES 1
 #endif
 
+// ARM64 and ARM64EC must use fallback paths. Unwind tables for those architectures have not been implemented.
+#if UE_CALLSTACK_TRACE_USE_UNWIND_TABLES && PLATFORM_CPU_ARM_FAMILY
+	#undef UE_CALLSTACK_TRACE_USE_UNWIND_TABLES
+	#define UE_CALLSTACK_TRACE_USE_UNWIND_TABLES 0
+#endif
 
 #if UE_CALLSTACK_TRACE_USE_UNWIND_TABLES
 
@@ -886,8 +891,8 @@ public:
 	FBacktracer(FMalloc* InMalloc);
 	~FBacktracer();
 	static FBacktracer*	Get();
-	inline uint32 GetBacktraceId(void* AddressOfReturnAddress);
-	uint32 GetBacktraceId(uint64 ReturnAddress);
+	inline uint32 GetBacktraceId(const void* AddressOfReturnAddress, uint32 FrameSkip);
+	uint32 GetBacktraceId(uint64 ReturnAddress, uint32 FrameSkip);
 	void AddModule(UPTRINT Base, const TCHAR* Name) {}
 	void RemoveModule(UPTRINT Base) {}
 
@@ -920,14 +925,14 @@ FBacktracer* FBacktracer::Get()
 }
 
 ////////////////////////////////////////////////////////////////////////////////
-uint32 FBacktracer::GetBacktraceId(void* AddressOfReturnAddress)
+uint32 FBacktracer::GetBacktraceId(const void* AddressOfReturnAddress, uint32 FrameSkip)
 {
-	const uint64 ReturnAddress = *(uint64*)AddressOfReturnAddress;
-	return GetBacktraceId(ReturnAddress);
+	const uint64 ReturnAddress = *(const uint64*)AddressOfReturnAddress;
+	return GetBacktraceId(ReturnAddress, FrameSkip);
 }
 
 ////////////////////////////////////////////////////////////////////////////////
-uint32 FBacktracer::GetBacktraceId(uint64 ReturnAddress)
+uint32 FBacktracer::GetBacktraceId(uint64 ReturnAddress, uint32 FrameSkip)
 {
 #if !UE_BUILD_SHIPPING
 	uint64 StackFrames[UE_CALLSTACK_TRACE_MAX_FRAMES];
@@ -961,8 +966,17 @@ uint32 FBacktracer::GetBacktraceId(uint64 ReturnAddress)
 
 		// Save the collected id
 		BacktraceEntry.Hash = BacktraceId;
-		BacktraceEntry.FrameCount = FrameIdx;
-		BacktraceEntry.Frames = StackFrames;
+
+		if (FrameIdx > FrameSkip)
+		{
+			BacktraceEntry.FrameCount = FrameIdx - FrameSkip;
+			BacktraceEntry.Frames = StackFrames + FrameSkip;
+		}
+		else
+		{
+			BacktraceEntry.FrameCount = FrameIdx;
+			BacktraceEntry.Frames = StackFrames;
+		}
 
 		// Add to queue to be processed. This might block until there is room in the
 		// queue (i.e. the processing thread has caught up processing).
@@ -1064,6 +1078,12 @@ FORCENOINLINE static uint32 CallstackTrace_GetCurrentId_Private()
 		return 0;
 	}
 
+	// We don't need the parent function in the traced callstack which will be
+	// CallstackTrace_GetCurrentId() or possibly MemoryTrace_Alloc/Free (inlining)
+	const uint32 FrameSkip = 1;
+
+#if UE_CALLSTACK_TRACE_USE_UNWIND_TABLES
+
 	const UPTRINT* StackPtr = CallstackTrace_GetStackFrame();
 	if (StackPtr == nullptr)
 	{
@@ -1071,12 +1091,20 @@ FORCENOINLINE static uint32 CallstackTrace_GetCurrentId_Private()
 	}
 	// It is expected that StackPtr == RSP at this point.
 
-	// We don't need the parent function in the traced callstack which will be
-	// CallstackTrace_GetCurrentId() or possibly MemoryTrace_Alloc/Free (inlining)
-	const uint32 FrameSkip = 1;
-
 	UPTRINT ProgCounter = UPTRINT(&CallstackTrace_GetCurrentId_Private);
 	return Instance->GetBacktraceId(ProgCounter, StackPtr, FrameSkip);
+
+#else
+
+	void* StackAddress = PLATFORM_RETURN_ADDRESS_FOR_CALLSTACKTRACING();
+
+#if PLATFORM_USE_CALLSTACK_ADDRESS_POINTER
+		return Instance->GetBacktraceId(StackAddress, FrameSkip);
+#else
+		return Instance->GetBacktraceId((uint64)StackAddress, FrameSkip);
+#endif
+
+#endif
 }
 
 ////////////////////////////////////////////////////////////////////////////////
diff --git a/Engine/Source/Runtime/Core/Public/ProfilingDebugging/CallstackTrace.h b/Engine/Source/Runtime/Core/Public/ProfilingDebugging/CallstackTrace.h
--- a/Engine/Source/Runtime/Core/Public/ProfilingDebugging/CallstackTrace.h
+++ b/Engine/Source/Runtime/Core/Public/ProfilingDebugging/CallstackTrace.h
@@ -11,7 +11,7 @@
 ////////////////////////////////////////////////////////////////////////////////
 #if !defined(UE_CALLSTACK_TRACE_ENABLED)
 	#if UE_TRACE_ENABLED && !UE_BUILD_SHIPPING
-		#if PLATFORM_WINDOWS && !PLATFORM_CPU_ARM_FAMILY
+		#if PLATFORM_WINDOWS
 			#define UE_CALLSTACK_TRACE_ENABLED 1
 		#endif
 	#endif
