Overhaul delegate lifetime checks to account for multicasts needing to be trivially relocatable

- Remove pointer tracking of multicasts, this cannot be made to work because multicasts need to be trivially relocatable for container use
- Remove automatic unregistration of stale delegates since we cannot keep pointers to the owning multicasts anymore
- Use delegate unregistrations to confirm correct cleanup instead of relying on the multicast's destructor to pass the same address as was used before
- Rename the intermediate layer class to clarify that this is a *tracking* feature that cannot fix the problems automatically

#rnx

#Synced-CL 50811283

#Synced-CL 50811283

diff --git a/Engine/Source/Runtime/Core/Private/Delegates/DelegateBase.cpp b/Engine/Source/Runtime/Core/Private/Delegates/DelegateBase.cpp
--- a/Engine/Source/Runtime/Core/Private/Delegates/DelegateBase.cpp
+++ b/Engine/Source/Runtime/Core/Private/Delegates/DelegateBase.cpp
@@ -7,60 +7,50 @@
 
 DEFINE_LOG_CATEGORY_STATIC(LogDelegates, Log, All);
 
-struct FDelegateRegistration
-{
-	FDelegateRegistration(FDelegateHandle InRegisteredHandle, FName InRegisteredFunctionName)
-		: RegisteredHandle(InRegisteredHandle), RegisteredFunctionName(InRegisteredFunctionName)
-	{}
+using FDelegatesRegistrationsToCallingFunctions = TMap<FDelegateHandle, FName>;
+using FMulticastRegistrationsPerModule = TMap<FName, FDelegatesRegistrationsToCallingFunctions>;
 
-	FDelegateHandle RegisteredHandle;
-	FName RegisteredFunctionName;
+struct FMulticastDelegateTrackerData
+{
+	FMulticastRegistrationsPerModule LiveRegistrationsPerModule;
+	FCriticalSection LiveRegistrationsPerModuleLock;
+	TMpscQueue<FDelegateHandle> UnregisteredDelegates;
 };
 
-using FMulticastRegistrations = TArray<FDelegateRegistration>;
-using FMulticastModuleUsers = TMap<FName, FMulticastRegistrations>;
-
-static FCriticalSection GLiveMulticastSetLock;
-static TMap<FMulticastDelegateSafetyLayer*, FMulticastModuleUsers> GLiveMulticastMap;
-static TAtomic<bool> GIsRemovingDelegates = false;
+FMulticastDelegateTrackerData& GMulticastDelegateTrackerData()
+{
+	static FMulticastDelegateTrackerData StaticData;
+	return StaticData;
+}
 
 static bool GEnableDelegateLifetimeChecks = !UE_BUILD_SHIPPING;
 static FAutoConsoleVariableRef EnableDelegateLifetimeChecksCVar(
 	TEXT("Delegates.EnableDelegateLifetimeChecks"),
 	GEnableDelegateLifetimeChecks,
-	TEXT("When set, enable the lifetime checking and automatic unregistration of delegates subscribed to multicast delegates."),
+	TEXT("When set, enable the lifetime checking of module-bound delegate instances subscribed to multicast delegates."),
 	ECVF_Default);
 
-TMpscQueue<FMulticastDelegateSafetyLayer*>& GDestroyedMulticastQueue()
+void FMulticastDelegateTracker::RegisterDelegateInstance(FDelegateHandle DelegateHandle, FName ModuleName, FName CallingFunctionName)
 {
-	static TMpscQueue<FMulticastDelegateSafetyLayer*> StaticQueue;
-	return StaticQueue;
-}
-
-FMulticastDelegateSafetyLayer::~FMulticastDelegateSafetyLayer()
-{
-	if (GEnableDelegateLifetimeChecks)
+	if (ModuleName != NAME_None && GEnableDelegateLifetimeChecks)
 	{
-		TRACE_CPUPROFILER_EVENT_SCOPE(~FMulticastDelegateSafetyLayer);
+		TRACE_CPUPROFILER_EVENT_SCOPE(RegisterDelegateInstance);
+
+		FMulticastDelegateTrackerData& TrackerData = GMulticastDelegateTrackerData();
 
-		checkf(!GIsRemovingDelegates, TEXT("Multicast destruction during delegate lifetime checks is unsafe!"));
-		GDestroyedMulticastQueue().Enqueue(this);
+		FScopeLock Lock(&TrackerData.LiveRegistrationsPerModuleLock);
+		FDelegatesRegistrationsToCallingFunctions& MulticastRegistrations = TrackerData.LiveRegistrationsPerModule.FindOrAdd(ModuleName);
+		MulticastRegistrations.Add(DelegateHandle, CallingFunctionName);
 	}
 }
 
-void FMulticastDelegateSafetyLayer::RegisterDelegateInstance(FDelegateHandle DelegateHandle, FName ModuleName, FName CallsiteName)
+void FMulticastDelegateTracker::UnregisterDelegateInstance(FDelegateHandle DelegateHandle)
 {
 	if (GEnableDelegateLifetimeChecks)
 	{
-		TRACE_CPUPROFILER_EVENT_SCOPE(RegisterDelegateInstance);
+		TRACE_CPUPROFILER_EVENT_SCOPE(UnregisterDelegateInstance);
 
-		if (ModuleName != NAME_None)
-		{
-			FScopeLock Lock(&GLiveMulticastSetLock);
-			FMulticastModuleUsers& MulticastModuleUsers = GLiveMulticastMap.FindOrAdd(this);
-			FMulticastRegistrations& MulticastRegistrations = MulticastModuleUsers.FindOrAdd(ModuleName);
-			MulticastRegistrations.Emplace(FDelegateRegistration(DelegateHandle, CallsiteName));
-		}
+		GMulticastDelegateTrackerData().UnregisteredDelegates.Enqueue(DelegateHandle);
 	}
 }
 
@@ -68,92 +58,67 @@ static void RemoveDestroyedMulticastsFromLiveSet()
 {
 	TRACE_CPUPROFILER_EVENT_SCOPE(RemoveDestroyedMulticastsFromLiveSet);
 
-	FScopeLock Lock(&GLiveMulticastSetLock);
+	FMulticastDelegateTrackerData& TrackerData = GMulticastDelegateTrackerData();
+	FScopeLock Lock(&TrackerData.LiveRegistrationsPerModuleLock);
 
-	FMulticastDelegateSafetyLayer* MulticastToDequeue;
-	while (GDestroyedMulticastQueue().Dequeue(MulticastToDequeue))
+	FDelegateHandle DelegateRegistrationToRemove;
+	while (TrackerData.UnregisteredDelegates.Dequeue(DelegateRegistrationToRemove))
 	{
-		GLiveMulticastMap.Remove(MulticastToDequeue);
+		for (TPair<FName, FDelegatesRegistrationsToCallingFunctions>& ModuleAndRegistrations : TrackerData.LiveRegistrationsPerModule)
+		{
+			if (FName* FoundFunctionName = ModuleAndRegistrations.Value.Find(DelegateRegistrationToRemove))
+			{
+				UE_LOG(LogDelegates, Verbose, TEXT("unregistered from %s: %s"), *ModuleAndRegistrations.Key.ToString(), *FoundFunctionName->ToString());
+				ModuleAndRegistrations.Value.Remove(DelegateRegistrationToRemove);
+				break;
+			}
+		}
 	}
 }
 
-void FMulticastDelegateSafetyLayer::RemoveAllDelegatesInModules(TConstArrayView<FName> ModuleNames)
+void FMulticastDelegateTracker::CheckForStaleDelegatesInModules(TConstArrayView<FName> ModuleNames)
 {
 	if (GEnableDelegateLifetimeChecks)
 	{
-		TRACE_CPUPROFILER_EVENT_SCOPE(RemoveAllDelegatesInModules);
-
-		int32 RemovedCount = 0;
-		int32 PendingDestroyedMulticastsCount = 0;
+		TRACE_CPUPROFILER_EVENT_SCOPE(CheckForStaleDelegatesInModules);
 
 		RemoveDestroyedMulticastsFromLiveSet();
 
-		FScopeLock Lock(&GLiveMulticastSetLock);
-
-		// First, remove all flushed multicasts. After this point, we won't check for more and assume no multicasts is destroyed in this time frame.
-		// This should be a safe assumption because almost all multicasts use the default unsafe policy, but it is checked in the destructor anyway.
-		GIsRemovingDelegates = true;
+		FMulticastDelegateTrackerData& TrackerData = GMulticastDelegateTrackerData();
+		FScopeLock Lock(&TrackerData.LiveRegistrationsPerModuleLock);
 
+		// Now find all live delegate registrations for each of those modules, and report them
+		for (FName ModuleName : ModuleNames)
 		{
-			TRACE_CPUPROFILER_EVENT_SCOPE(RemoveAllDelegatesInModules/Flush);
-
-			FMulticastDelegateSafetyLayer* MulticastToDequeue;
-			while (GDestroyedMulticastQueue().Dequeue(MulticastToDequeue))
+			if (FDelegatesRegistrationsToCallingFunctions* Registrations = TrackerData.LiveRegistrationsPerModule.Find(ModuleName))
 			{
-				GLiveMulticastMap.Remove(MulticastToDequeue);
-				PendingDestroyedMulticastsCount++;
-			}
-		}
-
-		// Now find all live delegate registrations for each of those modules, report them, and suppress them if needed.
-		for (auto LiveMulticastIt = GLiveMulticastMap.CreateIterator(); LiveMulticastIt; ++LiveMulticastIt)
-		{
-			FMulticastDelegateSafetyLayer* Multicast = LiveMulticastIt.Key();
-			FMulticastModuleUsers& MulticastModuleUsers = LiveMulticastIt.Value();
-
-			for (FName ModuleName : ModuleNames)
-			{
-				if (const FMulticastRegistrations* MulticastRegistrations = MulticastModuleUsers.Find(ModuleName))
+				for (const TPair<FDelegateHandle, FName>& DelegateRegistration : *Registrations)
 				{
-					for (const FDelegateRegistration& DelegateRegistration : *MulticastRegistrations)
-					{
-						if (Multicast->RemoveDelegateInstanceWrapper(DelegateRegistration.RegisteredHandle))
-						{
-							const FString ReportedMessage = FString::Printf(TEXT("stale delegate in %s: %s"), *ModuleName.ToString(), *DelegateRegistration.RegisteredFunctionName.ToString());
-							UE_LOG(LogDelegates, Error, TEXT("RemoveAllDelegatesInModules: %s"), *ReportedMessage);
-							ensureAlwaysMsgf(false, TEXT("%s"), *ReportedMessage);
-							RemovedCount++;
-						}
-					}
-
-					MulticastModuleUsers.Remove(ModuleName);
+					const FString ReportedMessage = FString::Printf(TEXT("stale delegate in %s: %s"), *ModuleName.ToString(), *DelegateRegistration.Value.ToString());
+					UE_LOG(LogDelegates, Error, TEXT("%s"), *ReportedMessage);
+					ensureAlwaysMsgf(false, TEXT("%s"), *ReportedMessage);
 				}
 			}
-
-			if (MulticastModuleUsers.Num() == 0)
-			{
-				LiveMulticastIt.RemoveCurrent();
-			}
 		}
 
-		GIsRemovingDelegates = false;
+		// Clear tracking data for every module we've checked so far
+		for (FName ModuleName : ModuleNames)
+		{
+			TrackerData.LiveRegistrationsPerModule.Remove(ModuleName);
+		}
 
 #if !UE_BUILD_SHIPPING
 
 		int32 TrackedDelegateInstances = 0;
-		size_t AllocatedSize = GLiveMulticastMap.GetAllocatedSize();
-		for (const TPair<FMulticastDelegateSafetyLayer*, FMulticastModuleUsers>& LiveMulticast : GLiveMulticastMap)
+		size_t AllocatedSize = TrackerData.LiveRegistrationsPerModule.GetAllocatedSize();
+		for (const TPair<FName, FDelegatesRegistrationsToCallingFunctions>& ModuleAndRegistrations : TrackerData.LiveRegistrationsPerModule)
 		{
-			AllocatedSize += LiveMulticast.Value.GetAllocatedSize();
-			for (const TPair<FName, FMulticastRegistrations>& MulticastModuleUser : LiveMulticast.Value)
-			{
-				TrackedDelegateInstances += MulticastModuleUser.Value.Num();
-				AllocatedSize += MulticastModuleUser.Value.GetAllocatedSize();
-			}
+			AllocatedSize += ModuleAndRegistrations.Value.GetAllocatedSize();
+			TrackedDelegateInstances += ModuleAndRegistrations.Value.Num();
 		}
 
-		UE_LOG(LogDelegates, Verbose, TEXT("RemoveAllDelegatesInModules: %d multicasts, %d in queue, %d instances, %dKB used for tracking"),
-			GLiveMulticastMap.Num(), PendingDestroyedMulticastsCount, TrackedDelegateInstances, AllocatedSize / 1024);
+		UE_LOG(LogDelegates, Verbose, TEXT("CheckForStaleDelegatesInModules: %d modules, %d tracked registrations, %dKB used for tracking"),
+			TrackerData.LiveRegistrationsPerModule.Num(), TrackedDelegateInstances, AllocatedSize / 1024);
 
 #endif // !UE_BUILD_SHIPPING
 	}
diff --git a/Engine/Source/Runtime/Core/Private/Modules/ModuleManager.cpp b/Engine/Source/Runtime/Core/Private/Modules/ModuleManager.cpp
--- a/Engine/Source/Runtime/Core/Private/Modules/ModuleManager.cpp
+++ b/Engine/Source/Runtime/Core/Private/Modules/ModuleManager.cpp
@@ -898,7 +898,7 @@ void FModuleManager::OnObjectCleanup()
 				UE_LOG(LogModuleManager, Log, TEXT("OnObjectCleanup: unloading dynamic library '%s'"), *LibraryAndModulesIterator.Key());
 
 #if UE_DELEGATE_CHECK_LIFETIME
-				FMulticastDelegateSafetyLayer::RemoveAllDelegatesInModules(LibraryAndModulesIterator.Value());
+				FMulticastDelegateTracker::CheckForStaleDelegatesInModules(LibraryAndModulesIterator.Value());
 #endif // UE_DELEGATE_CHECK_LIFETIME
 
 				const uint64_t InitialMemUsed = FPlatformMemory::GetMemoryUsedFast();
@@ -1345,7 +1345,7 @@ bool FModuleManager::UnloadModule( const FName InModuleName, bool bIsShutdown, b
 #if IS_MONOLITHIC
 
 #if UE_DELEGATE_CHECK_LIFETIME
-			FMulticastDelegateSafetyLayer::RemoveAllDelegatesInModules({ InModuleName });
+			FMulticastDelegateTracker::CheckForStaleDelegatesInModules({ InModuleName });
 #endif // UE_DELEGATE_CHECK_LIFETIME
 
 #else
diff --git a/Engine/Source/Runtime/Core/Public/Delegates/DelegateSignatureImpl.inl b/Engine/Source/Runtime/Core/Public/Delegates/DelegateSignatureImpl.inl
--- a/Engine/Source/Runtime/Core/Public/Delegates/DelegateSignatureImpl.inl
+++ b/Engine/Source/Runtime/Core/Public/Delegates/DelegateSignatureImpl.inl
@@ -38,7 +38,7 @@ template<typename UserPolicy> class TMulticastDelegateBase;
 #if UE_DELEGATE_CHECK_LIFETIME && defined(UE_DYNAMIC_MODULE_NAME)
 #include "Misc/SourceLocation.h"
 #include "Misc/SourceLocationUtils.h"
-#define UE_REGISTER_MULTICAST_INSTANCE(Handle) FMulticastDelegateSafetyLayer::RegisterDelegateInstance( \
+#define UE_REGISTER_MULTICAST_INSTANCE(Handle) FMulticastDelegateTracker::RegisterDelegateInstance( \
 	Handle,                                                                                          \
 	UE_DYNAMIC_MODULE_NAME,                                                                          \
 	UE::FSourceLocation::Current().GetFunctionName())
diff --git a/Engine/Source/Runtime/Core/Public/Delegates/MulticastDelegateBase.h b/Engine/Source/Runtime/Core/Public/Delegates/MulticastDelegateBase.h
--- a/Engine/Source/Runtime/Core/Public/Delegates/MulticastDelegateBase.h
+++ b/Engine/Source/Runtime/Core/Public/Delegates/MulticastDelegateBase.h
@@ -22,25 +22,25 @@
 
 #if UE_DELEGATE_CHECK_LIFETIME
 
-// Centralized registration layer for multicast delegates
-class FMulticastDelegateSafetyLayer
+// Centralized tracker for checking the lifetime of module-bound delegates registered to a multicast
+class FMulticastDelegateTracker
 {
 public:
 
-	virtual CORE_API ~FMulticastDelegateSafetyLayer();
-
 	// Register a delegate instance for a given module
-	CORE_API void RegisterDelegateInstance(FDelegateHandle DelegateHandle, FName ModuleName, FName CallsiteName);
-
-	// Remove every delegate instance for this module
-	static CORE_API void RemoveAllDelegatesInModules(TConstArrayView<FName> ModuleNames);
+	static CORE_API void RegisterDelegateInstance(FDelegateHandle DelegateHandle, FName ModuleName, FName CallsiteName);
 
-protected:
+	// Unregister a delegate instance 
+	static CORE_API void UnregisterDelegateInstance(FDelegateHandle DelegateHandle);
 
-	// We need this virtual wrapper because we're in a parent class for the actual delegate
-	virtual bool RemoveDelegateInstanceWrapper(FDelegateHandle Handle) = 0;
+	// Report stale delegates for a list of modules
+	static CORE_API void CheckForStaleDelegatesInModules(TConstArrayView<FName> ModuleNames);
 };
 
+#  define UE_UNREGISTER_MULTICAST_INSTANCE(Handle) FMulticastDelegateTracker::UnregisterDelegateInstance(Handle)
+
+#else
+#  define UE_UNREGISTER_MULTICAST_INSTANCE(Handle)
 #endif // UE_DELEGATE_CHECK_LIFETIME
 
 /**
@@ -48,9 +48,6 @@ protected:
  */
 template<typename UserPolicy>
 class TMulticastDelegateBase :
-#if UE_DELEGATE_CHECK_LIFETIME
-	protected FMulticastDelegateSafetyLayer,
-#endif // UE_DELEGATE_CHECK_LIFETIME
 	public TDelegateAccessHandlerBase<typename UserPolicy::FThreadSafetyMode>
 {
 protected:
@@ -71,6 +68,22 @@ public:
 		*this = MoveTemp(Other);
 	}
 
+#if UE_DELEGATE_CHECK_LIFETIME
+	~TMulticastDelegateBase()
+	{
+		FReadAccessScope ReadScope = GetReadAccessScope();
+
+		// We need multicasts to clean up after themselves so that unregistrations can be tracked
+		for (const UnicastDelegateType& DelegateBaseRef : InvocationList)
+		{
+			if (const IDelegateInstance* DelegateInstance = DelegateBaseRef.GetDelegateInstanceProtected())
+			{
+				UE_UNREGISTER_MULTICAST_INSTANCE(DelegateInstance->GetHandle());
+			}
+		}
+	}
+#endif // UE_DELEGATE_CHECK_LIFETIME
+
 	TMulticastDelegateBase& operator=(TMulticastDelegateBase&& Other)
 	{
 		if (&Other == this)
@@ -169,6 +182,8 @@ public:
 				IDelegateInstance* DelegateInstance = DelegateBaseRef.GetDelegateInstanceProtected();
 				if ((DelegateInstance != nullptr) && DelegateInstance->HasSameObject(InUserObject))
 				{
+					UE_UNREGISTER_MULTICAST_INSTANCE(DelegateInstance->GetHandle());
+
 					// Manually unbind the delegate here so the compaction will find and remove it.
 					DelegateBaseRef.Unbind();
 					++Result;
@@ -194,6 +209,11 @@ public:
 					|| DelegateInstance->HasSameObject(InUserObject)
 					|| DelegateInstance->IsCompactable())
 				{
+					if (DelegateInstance)
+					{
+						UE_UNREGISTER_MULTICAST_INSTANCE(DelegateInstance->GetHandle());
+					}
+
 					InvocationList.RemoveAtSwap(InvocationListIndex, EAllowShrinking::No);
 					++Result;
 				}
@@ -328,6 +348,8 @@ protected:
 	{
 		FWriteAccessScope WriteScope = GetWriteAccessScope();
 
+		UE_UNREGISTER_MULTICAST_INSTANCE(Handle);
+
 		for (int32 InvocationListIndex = 0; InvocationListIndex < InvocationList.Num(); ++InvocationListIndex)
 		{
 			UnicastDelegateType& DelegateBase = InvocationList[InvocationListIndex];
@@ -344,16 +366,6 @@ protected:
 		return false;
 	}
 
-#if UE_DELEGATE_CHECK_LIFETIME
-
-protected:
-	virtual bool RemoveDelegateInstanceWrapper(FDelegateHandle Handle) override
-	{
-		return RemoveDelegateInstance(Handle);
-	}
-
-#endif // UE_DELEGATE_CHECK_LIFETIME
-
 private:
 
 	inline bool ShouldSkipCompactionUnderAutoRTFM() const
@@ -401,6 +413,11 @@ private:
 
 			if (DelegateInstance == nullptr	|| DelegateInstance->IsCompactable())
 			{
+				if (DelegateInstance)
+				{
+					UE_UNREGISTER_MULTICAST_INSTANCE(DelegateInstance->GetHandle());
+				}
+
 				InvocationList.RemoveAtSwap(InvocationListIndex, EAllowShrinking::No);
 			}
 			else
@@ -466,6 +483,8 @@ private:
 	{
 		for (UnicastDelegateType& DelegateBaseRef : InvocationList)
 		{
+			UE_UNREGISTER_MULTICAST_INSTANCE(DelegateBaseRef.GetHandle());
+
 			DelegateBaseRef.Unbind();
 		}
 
