Fixed initialization of the static thread ID in FTaskTagScope

#Synced-CL 50253651

diff --git a/Engine/Source/Runtime/Core/Private/HAL/ThreadingBase.cpp b/Engine/Source/Runtime/Core/Private/HAL/ThreadingBase.cpp
--- a/Engine/Source/Runtime/Core/Private/HAL/ThreadingBase.cpp
+++ b/Engine/Source/Runtime/Core/Private/HAL/ThreadingBase.cpp
@@ -90,6 +90,10 @@ FTaskTagScope::FTaskTagScope(ETaskTag InTag) : Tag(InTag)
 	TagOnlyIfNone = true; // Now defaults to always optional
 	if (ActiveTaskTag == ETaskTag::EStaticInit)
 	{
+		// The static thread ID must be initialized from the first call to this constructor.
+		// Otherwise, there is a possibility that it will be initialized on a different thread.
+		GetStaticThreadId();
+
 		// Allow overwrite only from EStaticInit -> EGameThread
 		TagOnlyIfNone = Tag != ETaskTag::EGameThread;
 
