[CsvProfiler] Apply extension also when using CurrentCommand.Filename

[RN] CSVProfiler applies the .csv/.csv.gz extensions to captures started with an overridden/supplied filename
#Synced-CL 50831544

diff --git a/Engine/Source/Runtime/Core/Private/ProfilingDebugging/CsvProfiler.cpp b/Engine/Source/Runtime/Core/Private/ProfilingDebugging/CsvProfiler.cpp
--- a/Engine/Source/Runtime/Core/Private/ProfilingDebugging/CsvProfiler.cpp
+++ b/Engine/Source/Runtime/Core/Private/ProfilingDebugging/CsvProfiler.cpp
@@ -3811,7 +3811,8 @@ void FCsvProfiler::BeginCaptureInternal(const FCsvCaptureCommand& CurrentCommand
 
 	// Determine the output path and filename based on override params
 	FString DestinationFolder = CurrentCommand.DestinationFolder.IsEmpty() ? GetDefaultDirectory() : CurrentCommand.DestinationFolder + TEXT("/");
-	FString Filename = CurrentCommand.Filename.IsEmpty() ? FString::Printf(TEXT("Profile(%s)%s"), *FDateTime::Now().ToString(TEXT("%Y%m%d_%H%M%S")), CsvExtension) : CurrentCommand.Filename;
+	FString Filename = CurrentCommand.Filename.IsEmpty() ? FString::Printf(TEXT("Profile(%s)"), *FDateTime::Now().ToString(TEXT("%Y%m%d_%H%M%S"))) : CurrentCommand.Filename;
+	Filename = FPaths::SetExtension(Filename, CsvExtension);
 	OutputFilename = DestinationFolder + Filename;
 
 	TSharedPtr<FArchive> OutputFile = MakeShareable(IFileManager::Get().CreateFileWriter(*OutputFilename));
