The `IMPLEMENT_MODULE` injected initialize function should be disabled for transactions too.

#Synced-CL 50903561

diff --git a/Engine/Source/Runtime/Core/Public/Modules/ModuleManager.h b/Engine/Source/Runtime/Core/Public/Modules/ModuleManager.h
--- a/Engine/Source/Runtime/Core/Public/Modules/ModuleManager.h
+++ b/Engine/Source/Runtime/Core/Public/Modules/ModuleManager.h
@@ -946,7 +946,7 @@ namespace UE::Core::Private
 		/**/ \
 		/* @return	Returns an instance of this module */ \
 		/**/ \
-		static IModuleInterface* Initialize##ModuleName##Module() \
+		AUTORTFM_DISABLE static IModuleInterface* Initialize##ModuleName##Module() \
 		{ \
 			return new ModuleImplClass(); \
 		} \
