Unified Error - Fix move only error context not working.

#Synced-CL 50847432

diff --git a/Engine/Source/Runtime/Core/Public/Experimental/UnifiedError/UnifiedError.h b/Engine/Source/Runtime/Core/Public/Experimental/UnifiedError/UnifiedError.h
--- a/Engine/Source/Runtime/Core/Public/Experimental/UnifiedError/UnifiedError.h
+++ b/Engine/Source/Runtime/Core/Public/Experimental/UnifiedError/UnifiedError.h
@@ -346,11 +346,10 @@ namespace UE::UnifiedError
 	class TErrorDetails : public FDynamicErrorDetails
 	{
 	public:
-		TErrorDetails() {}
-		TErrorDetails(T&& InErrorDetail, const EDetailFilter InDetailFilterMask = EDetailFilter::Default)
+		explicit TErrorDetails(T&& InErrorDetail, const EDetailFilter InDetailFilterMask = EDetailFilter::Default)
+			: ErrorDetail(MoveTemp(InErrorDetail))
+			, DetailFilterMask(InDetailFilterMask)
 		{
-			ErrorDetail = MoveTemp(InErrorDetail);
-			DetailFilterMask = InDetailFilterMask;
 		}
 		/*TErrorDetails(T&& InErrorDetail, TRefCountPtr<const IErrorDetails> InInnerErrorDetails) : FDynamicErrorDetails(InInnerErrorDetails)
 		{
