[BUGFIX] Config File Iterator Fix while removing Sections Matching String

BUG:
The Iterator is being incremented while the FConfigCacheIni::EmptySection actually removes the entry from the map, and as the Iterator is not removing the entry itself via functions such as It.RemoveCurrent(), it leaves iterator in a bad state and while incrementing a check is invoked.

FIX:
- Deferred removal of sections allows using iterator without crashing.

#rnx

#Synced-CL 50908799

diff --git a/Engine/Source/Runtime/Core/Private/Misc/ConfigCacheIni.cpp b/Engine/Source/Runtime/Core/Private/Misc/ConfigCacheIni.cpp
--- a/Engine/Source/Runtime/Core/Private/Misc/ConfigCacheIni.cpp
+++ b/Engine/Source/Runtime/Core/Private/Misc/ConfigCacheIni.cpp
@@ -5051,13 +5051,21 @@ bool FConfigCacheIni::EmptySectionsMatchingString( const TCHAR* SectionString, c
 	{
 		bool bSaveOpsDisabled = bAreFileOperationsDisabled;
 		bAreFileOperationsDisabled = true;
-		for (FConfigFile::TIterator It(*File); It; ++It)
+
+		TArray<FString> SectionsToEmpty;
+		for (FConfigFile::TConstIterator It = File->CreateConstIterator(); It; ++It)
 		{
-			if (It.Key().Contains(SectionString) )
+			if (It.Key().Contains(SectionString))
 			{
-				bEmptied |= EmptySection(*(It.Key()), Filename);
+				SectionsToEmpty.Add(It.Key());
 			}
 		}
+				
+		for (const FString& SectionNameString : SectionsToEmpty)
+		{
+			bEmptied |= EmptySection(*SectionNameString, Filename);
+		}
+
 		bAreFileOperationsDisabled = bSaveOpsDisabled;
 	}
 	return bEmptied;
