Avoided re-normalizing the filename in FWindowsPlatformFile::GetFilenameOnDisk.

diff --git a/Engine/Source/Runtime/Core/Private/Windows/WindowsPlatformFile.cpp b/Engine/Source/Runtime/Core/Private/Windows/WindowsPlatformFile.cpp
--- a/Engine/Source/Runtime/Core/Private/Windows/WindowsPlatformFile.cpp
+++ b/Engine/Source/Runtime/Core/Private/Windows/WindowsPlatformFile.cpp
@@ -1297,9 +1297,12 @@ FString FWindowsPlatformFile::GetFilenameOnDisk(const TCHAR* Filename)
 		Path.ReplaceCharInline(TEXT('\\'), TEXT('/'), ESearchCase::CaseSensitive);
 	};
 
-	FString NormalizedFileName = *FNormalizedFilename(Filename);
+	FString Result = *FNormalizedFilename(Filename);
+	FString NormalizedInputPath = Result;
+	ConvertToUnrealPathInline(NormalizedInputPath);
+
 	TRACE_PLATFORMFILE_BEGIN_OPEN(Filename);
-	HANDLE hFile = CreateFile(*NormalizedFileName,
+	HANDLE hFile = CreateFile(*Result,
 		FILE_READ_ATTRIBUTES, FILE_SHARE_READ | FILE_SHARE_WRITE | FILE_SHARE_DELETE,
 		NULL, OPEN_EXISTING,
 		FILE_FLAG_BACKUP_SEMANTICS, // Required for opening a directory, which Filename might be
@@ -1308,8 +1311,7 @@ FString FWindowsPlatformFile::GetFilenameOnDisk(const TCHAR* Filename)
 	if (hFile == INVALID_HANDLE_VALUE)
 	{
 		TRACE_PLATFORMFILE_FAIL_OPEN(Filename);
-		ConvertToUnrealPathInline(NormalizedFileName);
-		return NormalizedFileName;
+		return NormalizedInputPath;
 	}
 	TRACE_PLATFORMFILE_END_OPEN(hFile);
 	ON_SCOPE_EXIT
@@ -1334,18 +1336,16 @@ FString FWindowsPlatformFile::GetFilenameOnDisk(const TCHAR* Filename)
 #endif
 	};
 
-	for (uint32 Length = NormalizedFileName.Len() + 10;;)
+	for (uint32 Length = Result.Len() + 10;;)
 	{
-		TArray<TCHAR, FString::AllocatorType>& CharArray = NormalizedFileName.GetCharArray();
+		TArray<TCHAR, FString::AllocatorType>& CharArray = Result.GetCharArray();
 		CharArray.SetNum(Length);
 
 		Length = GetFinalPathNameByHandle(hFile, CharArray.GetData(), CharArray.Num(), FILE_NAME_NORMALIZED);
 		if (Length == 0)
 		{
 			// Error - return normalized filename.
-			NormalizedFileName = *FNormalizedFilename(Filename);
-			ConvertToUnrealPathInline(NormalizedFileName);
-			return NormalizedFileName;
+			return NormalizedInputPath;
 		}
 		if (Length < (uint32)CharArray.Num())
 		{
@@ -1354,8 +1354,8 @@ FString FWindowsPlatformFile::GetFilenameOnDisk(const TCHAR* Filename)
 		}
 	}
 
-	ConvertToUnrealPathInline(NormalizedFileName);
-	return NormalizedFileName;
+	ConvertToUnrealPathInline(Result);
+	return Result;
 }
 
 ESymlinkResult FWindowsPlatformFile::IsSymlink(const TCHAR* Filename)
