Reduced the stack cost of a failed ensure from ~32 KiB to ~2 KiB

#Synced-CL 50253651

diff --git a/Engine/Source/Runtime/Core/Private/Misc/AssertionMacros.cpp b/Engine/Source/Runtime/Core/Private/Misc/AssertionMacros.cpp
--- a/Engine/Source/Runtime/Core/Private/Misc/AssertionMacros.cpp
+++ b/Engine/Source/Runtime/Core/Private/Misc/AssertionMacros.cpp
@@ -23,6 +23,7 @@
 #include "Misc/CommandLine.h"
 #include "Misc/OutputDeviceRedirector.h"
 #include "Misc/OutputDeviceError.h"
+#include "Misc/StringBuilder.h"
 #include "Stats/StatsMisc.h"
 #include "Misc/CoreDelegates.h"
 #include "HAL/ExceptionHandling.h"
@@ -542,8 +543,8 @@ FORCENOINLINE void FDebug::EnsureFailed(const ANSICHAR* Expr, const ANSICHAR* Fi
 		bool bShouldSendNewReport = false;
 
 		// Create a final string that we'll output to the log (and error history buffer)
-		TCHAR ErrorMsg[16384];
-		FCString::Snprintf(ErrorMsg, UE_ARRAY_COUNT(ErrorMsg), TEXT("Ensure condition failed: %hs " FILE_LINE_DESC_ANSI LINE_TERMINATOR_ANSI "%s" LINE_TERMINATOR_ANSI "Stack: " LINE_TERMINATOR_ANSI), Expr, File, Line, Msg);
+		TStringBuilder<1024> ErrorMsg;
+		ErrorMsg.Appendf(TEXT("Ensure condition failed: %hs " FILE_LINE_DESC_ANSI LINE_TERMINATOR_ANSI "%s" LINE_TERMINATOR_ANSI "Stack: " LINE_TERMINATOR_ANSI), Expr, File, Line, Msg);
 
 		// No debugger attached, so generate a call stack and submit a crash report
 		// Walk the stack and dump it to the allocated memory.
@@ -571,26 +572,26 @@ FORCENOINLINE void FDebug::EnsureFailed(const ANSICHAR* Expr, const ANSICHAR* Fi
 			}
 
 			// Also append the stack trace
-			FCString::StrncatTruncateDest(ErrorMsg, UE_ARRAY_COUNT(ErrorMsg), ANSI_TO_TCHAR(StackTrace));
+			ErrorMsg.Append(StackTrace);
 			FMemory::SystemFree(StackTrace);
 
 			// Dump the error and flush the log.
 #if !NO_LOGGING
-			FDebug::LogFormattedMessageWithCallstack(LogOutputDevice.GetCategoryName(), __FILE__, __LINE__, TEXT("=== Handled ensure: ==="), ErrorMsg, (GEnsuresAreErrors ? ELogVerbosity::Error : ELogVerbosity::Warning));
+			FDebug::LogFormattedMessageWithCallstack(LogOutputDevice.GetCategoryName(), __FILE__, __LINE__, TEXT("=== Handled ensure: ==="), *ErrorMsg, (GEnsuresAreErrors ? ELogVerbosity::Error : ELogVerbosity::Warning));
 #endif
 			GLog->Flush();
 
-			FCoreDelegates::OnEnsureFailed.Broadcast(Expr, File, Line, Msg, ErrorMsg);
+			FCoreDelegates::OnEnsureFailed.Broadcast(Expr, File, Line, Msg, *ErrorMsg);
 
 			// Trace the error
 #if !NO_LOGGING
 			if (GEnsuresAreErrors)
 			{
-				TRACE_LOG_MESSAGE(LogOutputDevice, Error, TEXT("%s"), ErrorMsg);
+				TRACE_LOG_MESSAGE(LogOutputDevice, Error, TEXT("%s"), *ErrorMsg);
 			}
 			else
 			{
-				TRACE_LOG_MESSAGE(LogOutputDevice, Warning, TEXT("%s"), ErrorMsg);
+				TRACE_LOG_MESSAGE(LogOutputDevice, Warning, TEXT("%s"), *ErrorMsg);
 			}
 #endif
 			if (FParse::Param(FCommandLine::Get(), TEXT("nocrashreports")))
@@ -616,7 +617,7 @@ FORCENOINLINE void FDebug::EnsureFailed(const ANSICHAR* Expr, const ANSICHAR* Fi
 					// Compute CRC of error string.  Note that along with the call stack, this includes the message
 					// string passed to the macro, so only truly redundant errors will go unreported.  Though it also
 					// means you shouldn't pass loop counters to ensureMsgf(), otherwise failures may spam the server!
-					const uint32 ErrorStrCRC = FCrc::StrCrc_DEPRECATED(ErrorMsg);
+					const uint32 ErrorStrCRC = FCrc::StrCrc_DEPRECATED(*ErrorMsg);
 
 					for (uint32 CurErrorIndex = 0; CurErrorIndex < StaticPreviousErrorCount; ++CurErrorIndex)
 					{
@@ -640,7 +641,7 @@ FORCENOINLINE void FDebug::EnsureFailed(const ANSICHAR* Expr, const ANSICHAR* Fi
 
 						FCoreDelegates::OnHandleSystemEnsure.Broadcast();
 
-						FPlatformMisc::SubmitErrorReport(ErrorMsg, EErrorReportMode::Balloon);
+						FPlatformMisc::SubmitErrorReport(*ErrorMsg, EErrorReportMode::Balloon);
 
 						bShouldSendNewReport = true;
 					}
@@ -683,7 +684,7 @@ FORCENOINLINE void FDebug::EnsureFailed(const ANSICHAR* Expr, const ANSICHAR* Fi
 
 #if PLATFORM_USE_REPORT_ENSURE
 			UE::TUniqueLock Lock(FailDebugMutex);
-			ReportEnsure(ErrorMsg, ProgramCounter);
+			ReportEnsure(*ErrorMsg, ProgramCounter);
 
 			GErrorHist[0] = TEXT('\0');
 			GErrorExceptionDescription[0] = TEXT('\0');
diff --git a/Engine/Source/Runtime/Core/Private/Misc/StringBuilder.cpp b/Engine/Source/Runtime/Core/Private/Misc/StringBuilder.cpp
--- a/Engine/Source/Runtime/Core/Private/Misc/StringBuilder.cpp
+++ b/Engine/Source/Runtime/Core/Private/Misc/StringBuilder.cpp
@@ -69,8 +69,9 @@ TStringBuilderBase<C>& TStringBuilderBase<C>::AppendV(const C* Fmt, va_list Args
 		const int32 RequiredSize = TCString<C>::FormatV(CurPos, RemainingSize, Fmt, ArgPack);
 		va_end(ArgPack);
 
-		if (!ensure(RequiredSize >= 0))
+		if (RequiredSize < 0)
 		{
+			UE_DEBUG_BREAK();
 			return *this;
 		}
 
