Stopped FWindowsPlatformFile::GetFilenameOnDisk() from following symbolic links.

#Synced-CL 50651673

diff --git a/Engine/Source/Runtime/Core/Private/Windows/WindowsPlatformFile.cpp b/Engine/Source/Runtime/Core/Private/Windows/WindowsPlatformFile.cpp
--- a/Engine/Source/Runtime/Core/Private/Windows/WindowsPlatformFile.cpp
+++ b/Engine/Source/Runtime/Core/Private/Windows/WindowsPlatformFile.cpp
@@ -1355,6 +1355,83 @@ FString FWindowsPlatformFile::GetFilenameOnDisk(const TCHAR* Filename)
 	}
 
 	ConvertToUnrealPathInline(Result);
+
+	// Check that the result only differs by case for the full path - otherwise do a slower approach which builds the path up by component.
+	// This is to ensure symbolic links aren't followed.
+	if (!Result.Equals(NormalizedInputPath, ESearchCase::IgnoreCase))
+	{
+		Result.Reset();
+
+		FString AccumulatedFilename;
+		AccumulatedFilename.Reserve(Result.Len());
+
+		TArray<uint8> FileNameInfoTempBuffer;
+		FileNameInfoTempBuffer.AddUninitialized(sizeof(FILE_NAME_INFO) + NormalizedInputPath.Len() * sizeof(TCHAR));
+		FPathViews::IterateComponents(
+			NormalizedInputPath,
+			[&AccumulatedFilename, &FileNameInfoTempBuffer, &Result](FStringView Component)
+			{
+				AccumulatedFilename /= Component;
+
+				if (Result.IsEmpty())
+				{
+					Result = Component;
+					if (Component.EndsWith(':'))
+					{
+						Result.ToUpperInline();
+					}
+					return;
+				}
+
+				HANDLE hFile = CreateFile(
+					*AccumulatedFilename,
+					FILE_READ_ATTRIBUTES, FILE_SHARE_READ | FILE_SHARE_WRITE | FILE_SHARE_DELETE,
+					NULL,
+					OPEN_EXISTING,
+					FILE_FLAG_OPEN_REPARSE_POINT | FILE_FLAG_BACKUP_SEMANTICS, // FILE_FLAG_OPEN_REPARSE_POINT prevents symlinks from being followed. FILE_FLAG_BACKUP_SEMANTICS is required for opening a directory, which AccumulatedFilename might be.
+					NULL
+				);
+				if (hFile == INVALID_HANDLE_VALUE)
+				{
+					// File is probably missing - leave Component unchanged
+					Result /= Component;
+					return;
+				}
+				ON_SCOPE_EXIT
+				{
+					CloseHandle(hFile);
+				};
+
+				for (;;)
+				{
+					FILE_NAME_INFO* BufferPtr = (FILE_NAME_INFO*)FileNameInfoTempBuffer.GetData();
+					BOOL bResult = GetFileInformationByHandleEx(hFile, FileNameInfo, BufferPtr, (DWORD)FileNameInfoTempBuffer.Num());
+					if (bResult)
+					{
+						Result /= FPathViews::GetPathLeaf(MakeStringView(BufferPtr->FileName, BufferPtr->FileNameLength / sizeof(*BufferPtr->FileName)));
+						return;
+					}
+
+					// If the error wasn't due to a too-small buffer - leave Component unchanged
+					if (GetLastError() != ERROR_MORE_DATA)
+					{
+						Result /= Component;
+						return;
+					}
+
+					// Increase buffer size and try again
+					FileNameInfoTempBuffer.SetNumUninitialized(sizeof(FILE_NAME_INFO) + BufferPtr->FileNameLength);
+				}
+			}
+		);
+
+		// If it's still not a case variant, just return the normalized input.
+		if (!Result.Equals(NormalizedInputPath, ESearchCase::IgnoreCase))
+		{
+			return NormalizedInputPath;
+		}
+	}
+
 	return Result;
 }
 
