UE: In Windows, call RandInit() in every thread since srand() is thread-local.

#Synced-CL 50906850

diff --git a/Engine/Source/Runtime/Core/Private/Windows/WindowsRunnableThread.cpp b/Engine/Source/Runtime/Core/Private/Windows/WindowsRunnableThread.cpp
--- a/Engine/Source/Runtime/Core/Private/Windows/WindowsRunnableThread.cpp
+++ b/Engine/Source/Runtime/Core/Private/Windows/WindowsRunnableThread.cpp
@@ -3,6 +3,7 @@
 #include "WindowsRunnableThread.h"
 #include "GenericPlatform/GenericPlatformCrashContext.h"
 #include "HAL/ExceptionHandling.h"
+#include "Misc/App.h"
 #include "Misc/FeedbackContext.h"
 #include "Misc/OutputDeviceError.h"
 #include "Stats/Stats.h"
@@ -60,6 +61,13 @@ uint32 FRunnableThreadWin::GuardedRun()
 	FPlatformProcess::SetThreadAffinityMask(ThreadAffinityMask);
 
 	FPlatformProcess::SetThreadName(*ThreadName);
+
+	// Initialize random number generator. srand() is thread-local in windows.
+	if (!FApp::bUseFixedSeed)
+	{
+		FMath::RandInit(FPlatformTime::Cycles());
+	}
+
 	const TCHAR* CmdLine = ::GetCommandLineW();
 	bool bNoExceptionHandler = FParse::Param(::GetCommandLineW(), TEXT("noexceptionhandler"));
 #if UE_BUILD_DEBUG
