[CrashFix] In debug a structured archive range check fails due to the use of INDEX_NONE in overridable serialization record entries. This change silences that check

#Synced-CL 50922333

diff --git a/Engine/Source/Runtime/Core/Private/Serialization/StructuredArchive.cpp b/Engine/Source/Runtime/Core/Private/Serialization/StructuredArchive.cpp
--- a/Engine/Source/Runtime/Core/Private/Serialization/StructuredArchive.cpp
+++ b/Engine/Source/Runtime/Core/Private/Serialization/StructuredArchive.cpp
@@ -159,7 +159,7 @@ void FStructuredArchive::SetScope(UE::StructuredArchive::Private::FSlotPosition
 				break;
 			case UE::StructuredArchive::Private::EElementType::Array:
 #if DO_STRUCTURED_ARCHIVE_CONTAINER_CHECKS
-				checkf(GetUnderlyingArchive().IsLoading() || CurrentContainer.Top()->Index <= CurrentContainer.Top()->Count, TEXT("Incorrect number of elements serialized in array"));
+				checkf(GetUnderlyingArchive().IsLoading() || CurrentContainer.Top()->Index <= CurrentContainer.Top()->Count || CurrentContainer.Top()->Count == INDEX_NONE, TEXT("Incorrect number of elements serialized in array"));
 #endif
 				Formatter.LeaveArray();
 #if DO_STRUCTURED_ARCHIVE_CONTAINER_CHECKS
