Replaced findClass with loadClass in the jni ClassLoader. Made improvements to jni TConstructor to show parameter names. Added some Marshaling to FEnv

#android

#Synced-CL 50898235

diff --git a/Engine/Plugins/Media/ElectraUtil/Source/ElectraSamples/Private/Android/ElectraTextureSample_Android.cpp b/Engine/Plugins/Media/ElectraUtil/Source/ElectraSamples/Private/Android/ElectraTextureSample_Android.cpp
--- a/Engine/Plugins/Media/ElectraUtil/Source/ElectraSamples/Private/Android/ElectraTextureSample_Android.cpp
+++ b/Engine/Plugins/Media/ElectraUtil/Source/ElectraSamples/Private/Android/ElectraTextureSample_Android.cpp
@@ -85,7 +85,7 @@ namespace UE::Jni
 	{
 		static constexpr FAnsiStringView ClassName = "com/epicgames/unreal/ElectraTextureSample";
 
-		inline static TConstructor<FElectraTextureSample> New;
+		inline static TConstructor<FElectraTextureSample()> New;
 
 		static void JNICALL nativeSignalSurfaceReadEvent(JNIEnv* env, jobject thiz, jlong InParentHandle);
 
diff --git a/Engine/Source/Runtime/ApplicationCore/Public/Android/AndroidJavaMediaPlayer.h b/Engine/Source/Runtime/ApplicationCore/Public/Android/AndroidJavaMediaPlayer.h
--- a/Engine/Source/Runtime/ApplicationCore/Public/Android/AndroidJavaMediaPlayer.h
+++ b/Engine/Source/Runtime/ApplicationCore/Public/Android/AndroidJavaMediaPlayer.h
@@ -43,8 +43,8 @@ namespace UE::Jni
 	struct FMediaPlayer14: Java::Lang::FObject
 	{
 		static constexpr FAnsiStringView ClassName = "com/epicgames/unreal/MediaPlayer14";
-		
-		inline static TConstructor<FMediaPlayer14, jboolean, jboolean, jboolean> New;
+
+		inline static TConstructor<FMediaPlayer14(jboolean SwizzlePixels, jboolean VulkanRenderer, jboolean NeedTrackInfo)> New;
 
 		static jint JNICALL nativeReadAt(JNIEnv* env, jobject thiz, jlong Identifier, jlong Position, Nio::FByteBuffer* Buffer, jint Offset, jint Count);
 
diff --git a/Engine/Source/Runtime/Core/Private/Android/AndroidJavaEnv.cpp b/Engine/Source/Runtime/Core/Private/Android/AndroidJavaEnv.cpp
--- a/Engine/Source/Runtime/Core/Private/Android/AndroidJavaEnv.cpp
+++ b/Engine/Source/Runtime/Core/Private/Android/AndroidJavaEnv.cpp
@@ -42,7 +42,7 @@ jclass AndroidJavaEnv::FindJavaClass(const char* name)
 	}
 
 	jstring ClassNameObj = Env->NewStringUTF(name);
-	jclass FoundClass = UE::Jni::Java::Lang::FClassLoader::findClass(ClassLoader, ClassNameObj).Leak();
+	jclass FoundClass = UE::Jni::Java::Lang::FClassLoader::loadClass(ClassLoader, ClassNameObj).Leak();
 
 	if (Env->ExceptionCheck())
 	{
@@ -66,7 +66,7 @@ jclass AndroidJavaEnv::FindJavaClassGlobalRef(const char* name)
 	}
 
 	jstring ClassNameObj = Env->NewStringUTF(name);
-	auto FoundClass = UE::Jni::Java::Lang::FClassLoader::findClass(ClassLoader, ClassNameObj);
+	auto FoundClass = UE::Jni::Java::Lang::FClassLoader::loadClass(ClassLoader, ClassNameObj);
 	
 	if (Env->ExceptionCheck())
 	{
@@ -231,7 +231,7 @@ namespace UE::Jni
 
 	FScopedJavaObject<jclass> FEnv::FindClass(const char* ClassName) const
 	{
-		FScopedJavaObject<jclass> Class = Java::Lang::FClassLoader::findClass(ClassLoader, *FScopedJavaObject{Value->NewStringUTF(ClassName)});
+		FScopedJavaObject<jclass> Class = Java::Lang::FClassLoader::loadClass(ClassLoader, *FScopedJavaObject{Value->NewStringUTF(ClassName)});
 
 		if (!Class)
 		{
@@ -323,8 +323,31 @@ namespace UE::Jni
 
 		return true;
 	}
+
+	FScopedJavaObject<Java::Lang::FString*> FEnv::Marshal(FStringView Param) const
+	{
+		return NewScopedJavaObject(Value, Value->NewString(reinterpret_cast<const jchar*>(Param.GetData()), Param.Len()));
+	}
 	
-    extern "C" void (FEnv::*const __start_UE_Jni_Initialize [[gnu::weak]])() const;
+	FString FEnv::Marshal(Java::Lang::FString* Param) const
+	{
+		FString String;
+
+		if (Param)
+		{
+			if (jsize Len = Value->GetStringLength(Param))
+			{
+				String.GetCharArray().AddUninitialized(Len + 1);
+				TCHAR* Buffer = String.GetCharArray().GetData();
+				Value->GetStringRegion(Param, 0, Len, reinterpret_cast<jchar*>(Buffer));
+				Buffer[Len] = '\0';
+			}
+		}
+
+		return String;
+	}
+
+	extern "C" void (FEnv::*const __start_UE_Jni_Initialize [[gnu::weak]])() const;
     extern "C" void (FEnv::*const __stop_UE_Jni_Initialize [[gnu::weak]])() const;
 	void FEnv::Initialize() const
 	{
@@ -339,6 +362,9 @@ namespace UE::Jni
 
 		ClassLoader = NewGlobalRef(*Get<FClassLoader* ()>(*FScopedJavaObject<TClass<FClass>*>{Value->FindClass(ClassName<FClass>.GetData())}, "getClassLoader")(Class<FGameActivity>));
 		
+		Initialize<FInteger>();
+		Initialize<Java::Util::FHashMap>();
+
 		for (auto* Initialize = &__start_UE_Jni_Initialize; Initialize < &__stop_UE_Jni_Initialize; ++Initialize)
 		{
 			(this->**Initialize)();
diff --git a/Engine/Source/Runtime/Core/Public/Android/AndroidJavaEnv.h b/Engine/Source/Runtime/Core/Public/Android/AndroidJavaEnv.h
--- a/Engine/Source/Runtime/Core/Public/Android/AndroidJavaEnv.h
+++ b/Engine/Source/Runtime/Core/Public/Android/AndroidJavaEnv.h
@@ -109,6 +109,7 @@ namespace UE::Jni
 		};
 		
 		template <typename T, typename... TArgs>
+		UE_DEPRECATED(5.8, "TConstructor<T(TArgs...) instead.")
 		constexpr FMember(TConstructor<T, TArgs...>& Constructor, const char* = nullptr)
 			: Name{"<init>"}
 			, Signature{Jni::Signature<void (TArgs...)>.GetData()}
@@ -117,6 +118,15 @@ namespace UE::Jni
 		{
 		}
 		
+		template <typename T, typename... TArgs>
+		constexpr FMember(TConstructor<T(TArgs...)>& Constructor, const char* = nullptr)
+			: Name{"<init>"}
+			, Signature{Jni::Signature<void (TArgs...)>.GetData()}
+			, Type{Method}
+			, MethodId{&Constructor.Value}
+		{
+		}
+		
 		template <typename T, typename U> requires std::is_function_v<U>
 		constexpr FMember(TMember<T, U>& Member, const char* _Name)
 			: Name{_Name}
@@ -443,6 +453,25 @@ namespace UE::Jni
 		bool Register(jclass Class, TArrayView<const FMember> Members) const;
 		bool Register(jclass Class, TArrayView<const FNativeMethod> NativeMethods) const;
 
+		template <CContiguousRange T> requires(!std::is_convertible_v<const T&, FStringView>)
+		auto Marshal(const T& Params) const
+		{
+			using U = decltype(Marshal(Params[0]))::Type;
+			
+			int32 Count = GetNum(Params);			
+			FScopedJavaObject<Java::Lang::TArray<U>*> OutArray{reinterpret_cast<Java::Lang::TArray<U>*>(Value->NewObjectArray(Count, Class<std::remove_pointer_t<U>>, nullptr))};
+			
+			for (int32 Index = 0; Index < Count; ++Index)
+			{
+				Value->SetObjectArrayElement(*OutArray, Index, *Marshal(Params[Index]));
+			}
+			
+			return OutArray;
+		}
+
+		FScopedJavaObject<Java::Lang::FString*> Marshal(FStringView Param) const;
+		FString Marshal(Java::Lang::FString* Param) const;
+
 		[[nodiscard]] JNIEnv* Get() const
 		{
 			return Value;
@@ -476,8 +505,13 @@ namespace UE::Jni
 		static constexpr void (FEnv::*Function)() const = &FEnv::Initialize<T>;
 	};
 
+	template <typename T, typename... TArgs> requires(!std::is_function_v<T>)
+	struct UE_DEPRECATED(5.8, "TConstructor<T(TArgs...) instead.") TConstructor<T, TArgs...>: TConstructor<T(TArgs...)>
+	{
+	};
+
 	template <typename T, typename... TArgs>
-	struct TConstructor
+	struct TConstructor<T(TArgs...)>
 	{
 		jmethodID Value = nullptr;
 		
@@ -1038,20 +1072,52 @@ namespace UE::Jni
 		}
 	};
 
-	namespace Java::Lang
+	namespace Java
 	{
-		struct FClassLoader: FObject
+		namespace Lang
 		{
-			static constexpr FAnsiStringView ClassName = "java/lang/ClassLoader";
+			struct FClassLoader: FObject
+			{
+				static constexpr FAnsiStringView ClassName = "java/lang/ClassLoader";
 
-			inline static TMember<FClassLoader*, jclass (jstring)> findClass;
+				inline static TMember<FClassLoader*, jclass (jstring name)> loadClass;
 
-			static constexpr FMember Members[]
+				static constexpr FMember Members[]
+				{
+					UE_JNI_MEMBER(loadClass)
+				};
+			};
+			
+			struct FInteger: FObject
 			{
-				UE_JNI_MEMBER(findClass)
+				static constexpr FAnsiStringView ClassName = "java/lang/Integer";
+		
+				inline static TConstructor<FInteger(jint value)> New;
+		
+				static constexpr FMember Members[]
+				{
+					UE_JNI_MEMBER(New)
+				};
 			};
-		};
-	}	
+		}
+		
+		namespace Util
+		{
+			struct FHashMap: Lang::FObject
+			{
+				static constexpr FAnsiStringView ClassName = "java/util/HashMap";
+		
+				inline static TConstructor<FHashMap()> New;
+				inline static TMember<FHashMap*, Lang::FObject* (Lang::FObject* key, Lang::FObject* value)> put;
+		
+				static constexpr FMember Members[]
+				{
+					UE_JNI_MEMBER(New),
+					UE_JNI_MEMBER(put)
+				};
+			};
+		}
+	}
 }
 
 // Helper class that automatically calls DeleteLocalRef on the passed-in Java object when goes out of scope
@@ -1059,6 +1125,8 @@ template <typename T> requires std::is_pointer_v<T>
 class FScopedJavaObject
 {
 public:
+	using Type = T;
+	
 	FScopedJavaObject() = default;
 
 	constexpr FScopedJavaObject(std::nullptr_t)
diff --git a/Engine/Source/Runtime/WebBrowser/Private/Android/AndroidJniWebViewControl.h b/Engine/Source/Runtime/WebBrowser/Private/Android/AndroidJniWebViewControl.h
--- a/Engine/Source/Runtime/WebBrowser/Private/Android/AndroidJniWebViewControl.h
+++ b/Engine/Source/Runtime/WebBrowser/Private/Android/AndroidJniWebViewControl.h
@@ -68,7 +68,7 @@ namespace UE::Jni
 		
 		static constexpr FAnsiStringView ClassName = "com/epicgames/unreal/WebViewControl";
 
-		inline static TConstructor<FWebViewControl, jlong, jint, jint, jboolean, jboolean, jboolean, jboolean, jboolean, jboolean, jstring, jboolean> New;
+		inline static TConstructor<FWebViewControl(jlong InNativePtr, jint Width, jint Height, jboolean bSwizzlePixels, jboolean bVulkanRenderer, jboolean bEnableRemoteDebugging, jboolean bUseTransparency, jboolean bEnableDomStorage, jboolean bShouldUseBitmapRender, jstring UserAgentApplication, jboolean bEnableFloatingCloseButton)> New;
 
 		static void JNICALL FloatingCloseButtonPressed(JNIEnv* JEnv, jobject thiz);
 		
