FVirtualWinApiHooks now uses Detours library for trampolines.

The FTextSectionEditor class was manually implementing x64 trampolines to hook memory allocation functions on Windows. Rather than implement trampolines for both ARM64 and ARM64EC, this uses Microsoft's Detours library which they have already ported.

With this change, the code now works correctly under all three architectures.

#rnx

#Synced-CL 50877619

diff --git a/Engine/Source/Runtime/Core/Core.Build.cs b/Engine/Source/Runtime/Core/Core.Build.cs
--- a/Engine/Source/Runtime/Core/Core.Build.cs
+++ b/Engine/Source/Runtime/Core/Core.Build.cs
@@ -311,9 +311,14 @@ public class Core : ModuleRules
 
 		IWYUSupport = IWYUSupport.KeepAsIs;
 
-		if (Target.WindowsPlatform.bEnableInstrumentation)
+		if (Target.Platform.IsInGroup(UnrealPlatformGroup.Windows))
 		{
-			PrivateDependencyModuleNames.Add("Detours");
+			// WindowsMemoryTrace.cpp uses Detours in all non-shipping builds. bEnableInstrumentation is used by
+			// RaceDetectorWindows.cpp, and is allowed in Shipping builds.
+			if (Target.Configuration != UnrealTargetConfiguration.Shipping || Target.WindowsPlatform.bEnableInstrumentation)
+			{
+				PrivateDependencyModuleNames.Add("Detours");
+			}
 		}
 
 		PublicDefinitions.Add("UE_DELEGATE_CHECK_LIFETIME=" + (Target.bCheckDelegateLifetime ? "1" : "0"));
diff --git a/Engine/Source/Runtime/Core/Private/ProfilingDebugging/Microsoft/WindowsMemoryTrace.cpp b/Engine/Source/Runtime/Core/Private/ProfilingDebugging/Microsoft/WindowsMemoryTrace.cpp
--- a/Engine/Source/Runtime/Core/Private/ProfilingDebugging/Microsoft/WindowsMemoryTrace.cpp
+++ b/Engine/Source/Runtime/Core/Private/ProfilingDebugging/Microsoft/WindowsMemoryTrace.cpp
@@ -20,6 +20,7 @@
 #include <shellapi.h>
 #include <winternl.h>
 #include <intrin.h>
+#include <detours/detours.h>
 
 #if (NTDDI_VERSION >= NTDDI_WIN10_RS4)
 #pragma comment(lib, "mincore.lib") // VirtualAlloc2
@@ -40,193 +41,7 @@ static_assert(sizeof(FAddrPack) == sizeof(uint64), "");
 
 ////////////////////////////////////////////////////////////////////////////////
 #if defined(PLATFORM_SUPPORTS_TRACE_WIN32_VIRTUAL_MEMORY_HOOKS)
-class FTextSectionEditor
-{
-public:
-							~FTextSectionEditor();
-	template <typename T> T*Hook(T* Target, T* HookFunction);
-
-private:
-	struct FTrampolineBlock
-	{
-		FTrampolineBlock*	Next;
-		uint32				Size;
-		uint32				Used;
-	};
-
-	static void*			GetActualAddress(void* Function);
-	FTrampolineBlock*		AllocateTrampolineBlock(void* Reference);
-	uint8*					AllocateTrampoline(void* Reference, unsigned int Size);
-	void*					HookImpl(void* Target, void* HookFunction);
-	FTrampolineBlock*		HeadBlock = nullptr;
-};
-
-////////////////////////////////////////////////////////////////////////////////
-FTextSectionEditor::~FTextSectionEditor()
-{
-	for (FTrampolineBlock* Block = HeadBlock; Block != nullptr; Block = Block->Next)
-	{
-		DWORD Unused;
-		VirtualProtect(Block, Block->Size, PAGE_EXECUTE_READ, &Unused);
-	}
-
-	FlushInstructionCache(GetCurrentProcess(), nullptr, 0);
-}
-
-////////////////////////////////////////////////////////////////////////////////
-void* FTextSectionEditor::GetActualAddress(void* Function)
-{
-	// Follow a jmp instruction (0xff/4 only for now) at function and returns
-	// where it would jmp to.
-
-	uint8* Addr = (uint8*)Function;
-	int Offset = unsigned(Addr[0] & 0xf0) == 0x40; // REX prefix
-	if (Addr[Offset + 0] == 0xff && Addr[Offset + 1] == 0x25)
-	{
-		Addr += Offset;
-		Addr = *(uint8**)(Addr + 6 + *(uint32*)(Addr + 2));
-	}
-	return Addr;
-}
-
-////////////////////////////////////////////////////////////////////////////////
-FTextSectionEditor::FTrampolineBlock* FTextSectionEditor::AllocateTrampolineBlock(void* Reference)
-{
-	static const SIZE_T BlockSize = 0x10000; // 64KB is Windows' canonical granularity
-
-	// Find the start of the main allocation that mapped Reference
-	MEMORY_BASIC_INFORMATION MemInfo;
-	VirtualQuery(Reference, &MemInfo, sizeof(MemInfo));
-	auto* Ptr = (uint8*)(MemInfo.AllocationBase);
-
-	// Step backwards one block at a time and try and allocate that address
-	while (true)
-	{
-		Ptr -= BlockSize;
-		if (VirtualAlloc(Ptr, BlockSize, MEM_RESERVE|MEM_COMMIT, PAGE_READWRITE) != nullptr)
-		{
-			break;
-		}
-
-		UPTRINT Distance = UPTRINT(Reference) - UPTRINT(Ptr);
-		if (Distance >= 1ull << 31)
-		{
-			check(!"Failed to allocate trampoline blocks for memory tracing hooks");
-		}
-	}
-
-	auto* Block = (FTrampolineBlock*)Ptr;
-	Block->Next = HeadBlock;
-	Block->Size = BlockSize;
-	Block->Used = sizeof(FTrampolineBlock);
-	HeadBlock = Block;
-
-	return Block;
-}
-
-////////////////////////////////////////////////////////////////////////////////
-uint8* FTextSectionEditor::AllocateTrampoline(void* Reference, unsigned int Size)
-{
-	// Try and find a block that's within 2^31 bytes before Reference
-	FTrampolineBlock* Block;
-	for (Block = HeadBlock; Block != nullptr; Block = Block->Next)
-	{
-		UPTRINT Distance = UPTRINT(Reference) - UPTRINT(Block);
-		if (Distance < 1ull << 31)
-		{
-			break;
-		}
-	}
-
-	// If we didn't find a block then we need to allocate a new one.
-	if (Block == nullptr)
-	{
-		Block = AllocateTrampolineBlock(Reference);
-	}
-
-	// Allocate space for the trampoline.
-	uint32 NextUsed = Block->Used + Size;
-	if (NextUsed > Block->Size)
-	{
-		// Block is full. We could allocate a new block here but as it is not
-		// expected that so many hooks will be made this path shouldn't happen
-		check(!"Unable to allocate memory for memory tracing's hooks");
-	}
-
-	uint8* Out = (uint8*)Block + Block->Used;
-	Block->Used = NextUsed;
-
-	return Out;
-}
-
-////////////////////////////////////////////////////////////////////////////////
-template <typename T>
-T* FTextSectionEditor::Hook(T* Target, T* HookFunction)
-{
-	return (T*)HookImpl((void*)Target, (void*)HookFunction);
-}
-
-////////////////////////////////////////////////////////////////////////////////
-void* FTextSectionEditor::HookImpl(void* Target, void* HookFunction)
-{
-	Target = GetActualAddress(Target);
-
-	// Very rudimentary x86_64 instruction length decoding that only supports op
-	// code ranges (0x80,0x8b) and (0x50,0x5f). Enough for simple prologues
-	uint8* __restrict Start = (uint8*)Target;
-	const uint8* Read = Start;
-	do
-	{
-		Read += (Read[0] & 0xf0) == 0x40; // REX prefix
-		uint8 Inst = *Read++;
-		if (unsigned(Inst - 0x80) < 0x0cu)
-		{
-			uint8 ModRm = *Read++;
-			Read += ((ModRm & 0300) < 0300) & ((ModRm & 0007) == 0004); // SIB
-			switch (ModRm & 0300) // Disp[8|32]
-			{
-			case 0100: Read += 1; break;
-			case 0200: Read += 5; break;
-			}
-			Read += (Inst == 0x83);
-		}
-		else if (unsigned(Inst - 0x50) >= 0x10u)
-			check(!"Unknown instruction");
-	}
-	while (Read - Start < 6);
-
-	static const int TrampolineSize = 24;
-	int PatchSize = int(Read - Start);
-	uint8* TrampolinePtr = AllocateTrampoline(Start, PatchSize + TrampolineSize);
-
-	// Write the trampoline
-	*(void**)TrampolinePtr = HookFunction;
-
-	uint8* PatchJmp = TrampolinePtr + sizeof(void*);
-	memcpy(PatchJmp, Start, PatchSize);
-
-	PatchJmp += PatchSize;
-	*PatchJmp = 0xe9;
-	*(int32*)(PatchJmp + 1) = int32(intptr_t(Start + PatchSize) - intptr_t(PatchJmp)) - 5;
-
-	// Need to make the text section writeable
-	DWORD ProtPrev;
-	UPTRINT ProtBase = UPTRINT(Target) & ~0x0fff;						// 0x0fff is mask of VM page size
-	SIZE_T ProtSize = ((ProtBase + 16 + 0x1000) & ~0x0fff) - ProtBase;	// 16 is enough for one x86 instruction
-	VirtualProtect((void*)ProtBase, ProtSize, PAGE_EXECUTE_READWRITE, &ProtPrev);
-
-	// Patch function to jmp to the hook
-	uint16* HookJmp = (uint16*)Target;
-	HookJmp[0] = 0x25ff;
-	*(int32*)(HookJmp + 1) = int32(intptr_t(TrampolinePtr) - intptr_t(HookJmp + 3));
-
-	// Put the protection back the way it was
-	VirtualProtect((void*)ProtBase, ProtSize, ProtPrev, &ProtPrev);
 
-	return PatchJmp - PatchSize;
-}
-
-////////////////////////////////////////////////////////////////////////////////
 class FVirtualWinApiHooks
 {
 public:
@@ -274,18 +89,23 @@ void FVirtualWinApiHooks::Initialize(bool bInLight)
 {
 	bLight = bInLight;
 
-	FTextSectionEditor Editor;
-
 	// Note that hooking alloc functions is done last as applying the hook can
 	// allocate some memory pages.
 
-	VmFreeOrig = Editor.Hook(VirtualFree, &FVirtualWinApiHooks::VmFree);
-	VmFreeExOrig = Editor.Hook(VirtualFreeEx, &FVirtualWinApiHooks::VmFreeEx);
+	DetourTransactionBegin();
+	DetourUpdateThread(GetCurrentThread());
+
+	VmFreeOrig = VirtualFree;
+	DetourAttach(&VmFreeOrig, &FVirtualWinApiHooks::VmFree);
+
+	VmFreeExOrig = VirtualFreeEx;
+	DetourAttach(&VmFreeExOrig, &FVirtualWinApiHooks::VmFreeEx);
 
 #if PLATFORM_WINDOWS
 #if (NTDDI_VERSION >= NTDDI_WIN10_RS4)
 	{
-		VmAlloc2Orig = Editor.Hook(VirtualAlloc2, &FVirtualWinApiHooks::VmAlloc2);
+		VmAlloc2Orig = VirtualAlloc2;
+		DetourAttach(&VmAlloc2Orig, &FVirtualWinApiHooks::VmAlloc2);
 	}
 #else // NTDDI_VERSION
 	{
@@ -302,14 +122,23 @@ void FVirtualWinApiHooks::Initialize(bool bInLight)
 		}
 		if (VmAlloc2Orig)
 		{
-			VmAlloc2Orig = Editor.Hook(VmAlloc2Orig, &FVirtualWinApiHooks::VmAlloc2);
+			DetourAttach(&VmAlloc2Orig, &FVirtualWinApiHooks::VmAlloc2);
 		}
 	}
 #endif // NTDDI_VERSION
 #endif // PLATFORM_WINDOWS
 
-	VmAllocExOrig = Editor.Hook(VirtualAllocEx, &FVirtualWinApiHooks::VmAllocEx);
-	VmAllocOrig = Editor.Hook(VirtualAlloc, &FVirtualWinApiHooks::VmAlloc);
+	VmAllocExOrig = VirtualAllocEx;
+	DetourAttach(&VmAllocExOrig, &FVirtualWinApiHooks::VmAllocEx);
+
+	VmAllocOrig = VirtualAlloc;
+	DetourAttach(&VmAllocOrig, &FVirtualWinApiHooks::VmAlloc);
+
+	LONG Result = DetourTransactionCommit();
+	if (Result != NO_ERROR)
+	{
+		UE_LOG(LogMemory, Warning, TEXT("Memory trace failed to install hooks. Detours returned 0x%08x"), Result);
+	}
 }
 
 ////////////////////////////////////////////////////////////////////////////////
