Add some additional CSV metadata for android to help with perf regression investigations
#Synced-CL 50901525

diff --git a/Engine/Source/Runtime/Core/Private/Android/AndroidPlatformMisc.cpp b/Engine/Source/Runtime/Core/Private/Android/AndroidPlatformMisc.cpp
--- a/Engine/Source/Runtime/Core/Private/Android/AndroidPlatformMisc.cpp
+++ b/Engine/Source/Runtime/Core/Private/Android/AndroidPlatformMisc.cpp
@@ -62,6 +62,7 @@
 #include "IVulkanDynamicRHI.h"
 #include <vulkan/vulkan_android.h>
 #include "Android/AndroidPlatformProperties.h"
+#include "ProfilingDebugging/CsvProfiler.h"
 
 #include "AndroidVKQuality.h"
 
@@ -652,6 +653,11 @@ void FAndroidMisc::PlatformInit()
 	AndroidScudoMemoryTrace::Init();
 
 	FAndroidPlatformMemory::AndroidMemoryProbeStart();
+
+	CSV_METADATA( TEXT("Android_BuildNumber"), *GetDeviceBuildNumber() );
+	CSV_METADATA( TEXT("Android_DeviceMake"),  *GetDeviceMake() );
+	CSV_METADATA( TEXT("Android_DeviceModel"), *GetDeviceModel() );
+	CSV_METADATA( TEXT("Android_Version"), *GetAndroidVersion() );
 }
 
 extern void AndroidThunkCpp_DismissSplashScreen();
