Fix some race conditions with deleting async read requests in FFileHelper::LoadFileToStringAsync

#Synced-CL 50966172

diff --git a/Engine/Source/Runtime/Core/Private/Misc/FileHelper.cpp b/Engine/Source/Runtime/Core/Private/Misc/FileHelper.cpp
--- a/Engine/Source/Runtime/Core/Private/Misc/FileHelper.cpp
+++ b/Engine/Source/Runtime/Core/Private/Misc/FileHelper.cpp
@@ -21,6 +21,7 @@
 #include "Misc/ConfigCacheIni.h"
 #include "Misc/SecureHash.h"
 #include "HAL/FileManagerGeneric.h"
+#include "Tasks/Pipe.h"
 #include "Templates/Function.h"
 #include "Templates/UniquePtr.h"
 
@@ -293,6 +294,11 @@ bool FFileHelper::LoadFileToString(FString& Result, IPlatformFile* PlatformFile,
 	return bSuccess;
 }
 
+namespace UE::Core::Private
+{
+static UE::Tasks::FPipe CloseHandlePipe(TEXT("CloseFileHandlesAndRequests"));
+}
+
 UE::Tasks::TTask<TValueOrError<FString, FFileSystemError>> FFileHelper::LoadFileToStringAsync(const TCHAR* Filename, FLoadFileAsyncOptions Options)
 {
 	IPlatformFile& PlatformFile = Options.OverridePlatformFile ? *Options.OverridePlatformFile : FPlatformFileManager::Get().GetPlatformFile();
@@ -306,8 +312,10 @@ UE::Tasks::TTask<TValueOrError<FString, FFileSystemError>> FFileHelper::LoadFile
 	struct FState
 	{
 		TUniquePtr<IAsyncReadFileHandle> Handle;
-		// The previous request, either size or data, waiting to be deleted when we're no longer in the associated callback and it's safe to do so
-		TUniquePtr<IAsyncReadRequest> RequestToDelete;
+		// The size and data requests must be deleted after their callback has finished executing
+		// The callbacks may nest in the callstack if ReadRequest returns instantly with cached data so the SizeRequest cannot be deleted in the read callback.
+		TUniquePtr<IAsyncReadRequest> SizeRequest;
+		TUniquePtr<IAsyncReadRequest> DataRequest;
 		UE::Tasks::FTaskEvent DataEvent{TEXT("LoadFileToStringAsync_DataEvent")};
 		int64 Size = 0;
 		TUniquePtr<uint8[]> Data;
@@ -327,7 +335,7 @@ UE::Tasks::TTask<TValueOrError<FString, FFileSystemError>> FFileHelper::LoadFile
 	// The API doesn't require the supplied pointer to live past the function call, just that the copy remain safely callable.
 	TFunction<void(bool, IAsyncReadRequest*)> SizeCallback = [State](bool bCancelled, IAsyncReadRequest* Request) mutable
 	{
-		State->RequestToDelete.Reset(Request);
+		State->SizeRequest.Reset(Request);
 		State->Size = Request->GetSizeResults();
 		if (State->Size <= 0 || State->Size > (int64)MAX_int32) // Result is -1 on cancellation or error
 		{
@@ -337,7 +345,7 @@ UE::Tasks::TTask<TValueOrError<FString, FFileSystemError>> FFileHelper::LoadFile
 		}
 		TFunction<void(bool, IAsyncReadRequest*)> ReadCallback = [State](bool bCancelled, IAsyncReadRequest* Request) mutable
 		{
-			State->RequestToDelete.Reset(Request);
+			State->DataRequest.Reset(Request);
 			State->DataEvent.Trigger();
 		};
 		State->Data = MakeUnique<uint8[]>(State->Size);
@@ -352,6 +360,23 @@ UE::Tasks::TTask<TValueOrError<FString, FFileSystemError>> FFileHelper::LoadFile
 
 	return UE::Tasks::Launch(TEXT("LoadFileToStringAsync"), [State]() -> TValueOrError<FString, FFileSystemError>
 		{
+			// The start of this task can race with the IAsyncReadRequest code which calls the callback for us so that the completion flag is not properly set yet
+			// Possibly overkill, but if the handles can't be released immediately, let a task do it
+			// Put these tasks into a pipe so that an uncooperative API can only block a single worker
+			// The file handle also shouldn't be closed until the requests are deleted
+			UE::Core::Private::CloseHandlePipe.Launch(TEXT("DeleteReadRequests"), 
+				[SizeRequest=MoveTemp(State->SizeRequest),DataRequest=MoveTemp(State->DataRequest),Handle=MoveTemp(State->Handle)]() mutable
+				{
+					SizeRequest->EnsureCompletion();
+					SizeRequest.Reset();
+					if (DataRequest)
+					{
+						DataRequest->EnsureCompletion();
+						DataRequest.Reset();
+					}
+					Handle.Reset();
+				});
+
 			if (State->Size < 0)
 			{
 				// Async read handle doesn't provide detailed errors at this point
