Logging: Convert UE_LOG to use log records internally

This adds file, line, and context to every UE_LOG even though they are inherently unstructured.

#Synced-CL 50967057

diff --git a/Engine/Source/Runtime/Core/Private/Logging/StructuredLog.cpp b/Engine/Source/Runtime/Core/Private/Logging/StructuredLog.cpp
--- a/Engine/Source/Runtime/Core/Private/Logging/StructuredLog.cpp
+++ b/Engine/Source/Runtime/Core/Private/Logging/StructuredLog.cpp
@@ -48,8 +48,8 @@ void AppendQuotedJsonString(TStringBuilderBase<CharType>& Builder, FUtf8StringVi
 namespace UE::Logging::Private
 {
 
-// Temporary override until performance and functionality are sufficient for this to be the default.
-CORE_API bool GConvertBasicLogToLogRecord = false;
+// Temporarily provided to toggle back to the previous behavior from code.
+CORE_API bool GConvertBasicLogToLogRecord = true;
 
 // Experimental feature to prepend log context to the log message during formatting.
 CORE_API bool GPrependLogContextToLogMessage = false;
