Add OnGPUHang in the ThreadHeartBeat

#Synced-CL 50700053

diff --git a/Engine/Source/Runtime/Core/Private/HAL/ThreadHeartBeat.cpp b/Engine/Source/Runtime/Core/Private/HAL/ThreadHeartBeat.cpp
--- a/Engine/Source/Runtime/Core/Private/HAL/ThreadHeartBeat.cpp
+++ b/Engine/Source/Runtime/Core/Private/HAL/ThreadHeartBeat.cpp
@@ -34,6 +34,8 @@
 const double HangDetectorClock_MaxTimeStep_MS = 2000.0;
 const double HitchDetectorClock_MaxTimeStep_MS = 50.0;
 
+extern CORE_API bool GIsGPUCrashed;
+
 FThreadHeartBeatClock::FThreadHeartBeatClock(double InMaxTimeStep)
 	: MaxTimeStepCycles(FPlatformTime::SecondsToCycles64(InMaxTimeStep))
 {
@@ -288,6 +290,25 @@ void FORCENOINLINE FThreadHeartBeat::OnHang(double HangDuration, uint32 ThreadTh
 #endif // MINIMAL_FATAL_HANG_DETECTION == 0
 }
 
+void FORCENOINLINE FThreadHeartBeat::OnGPUHang(double HangDuration)
+{
+#if MINIMAL_FATAL_HANG_DETECTION
+	// We want to avoid all memory allocations if a hang is detected.
+	// Force a crash in a way that will generate a crash report.
+
+	// Avoiding calling RaiseException here will keep OnGPUHang on the top of the crash callstack,
+	// making crash bucketing easier when looking at retail crash dumps on supported platforms.
+
+	//FPlatformMisc::RaiseException(0xe0000003);
+	*((uint32*)3) = 0xe0000003;
+
+#elif UE_ASSERT_ON_HANG
+	UE_LOG(LogCore, Fatal, TEXT("GPU hang detected. A frame has not been presented for %.2f seconds."), HangDuration);
+#else
+	UE_LOG(LogCore, Error, TEXT("GPU hang detected. A frame has not been presented for %.2f seconds."), HangDuration);
+#endif
+}
+
 uint32 FThreadHeartBeat::Run()
 {
 	bool InHungState = false;
@@ -316,7 +337,11 @@ uint32 FThreadHeartBeat::Run()
 			// Only want to call this once per hang (particularly if we're just ensuring).
 			InHungState = true;
 
-			if (ThreadThatHung == FThreadHeartBeat::PresentThreadId)
+			if (GIsGPUCrashed)
+			{ 
+				OnGPUHang(HangDuration);
+			}
+			else if (ThreadThatHung == FThreadHeartBeat::PresentThreadId)
 			{
 				OnPresentHang(HangDuration);
 			}
diff --git a/Engine/Source/Runtime/Core/Public/HAL/ThreadHeartBeat.h b/Engine/Source/Runtime/Core/Public/HAL/ThreadHeartBeat.h
--- a/Engine/Source/Runtime/Core/Public/HAL/ThreadHeartBeat.h
+++ b/Engine/Source/Runtime/Core/Public/HAL/ThreadHeartBeat.h
@@ -150,6 +150,7 @@ class FThreadHeartBeat : public FRunnable
 
 	CORE_API void FORCENOINLINE OnHang(double HangDuration, uint32 ThreadThatHung);
 	CORE_API void FORCENOINLINE OnPresentHang(double HangDuration);
+	CORE_API void FORCENOINLINE OnGPUHang(double HangDuration);
 
 	CORE_API bool IsEnabled();
 
