Initialize the CSV profiler at the start of PreInit to capture more of the initial boot

#Synced-CL 50901525

diff --git a/Engine/Source/Runtime/Core/Private/ProfilingDebugging/CsvProfiler.cpp b/Engine/Source/Runtime/Core/Private/ProfilingDebugging/CsvProfiler.cpp
--- a/Engine/Source/Runtime/Core/Private/ProfilingDebugging/CsvProfiler.cpp
+++ b/Engine/Source/Runtime/Core/Private/ProfilingDebugging/CsvProfiler.cpp
@@ -300,6 +300,8 @@ static thread_local bool GCsvThreadLocalWaitsEnabled = false;
 
 static uint64 GCaptureStartTimestamp = 0;
 
+static bool GCsvProfilerInitialized = false;
+
 //
 // Forward declarations
 //
@@ -4826,8 +4828,30 @@ void FCsvProfiler::RecordEndOfPipeCustomStat(const FName& StatName, uint32 Categ
 #endif
 }
 
+void FCsvProfiler::PreInit(const TCHAR * Commandline)
+{
+#if CSV_PROFILER_ALLOW_DEBUG_FEATURES
+	// If we're doing boot profiling then we need to init early and make sure the commandline is set
+	int32 NumCsvFrames = 0;
+	if (FParse::Value(Commandline, TEXT("csvCaptureFrames="), NumCsvFrames))
+	{
+		if (!FCommandLine::IsInitialized())
+		{
+			FCommandLine::Set(Commandline);
+		}
+		FCsvProfiler::Get()->Init();
+	}
+#endif
+}
+
 void FCsvProfiler::Init()
 {
+	if (GCsvProfilerInitialized)
+	{
+		return;
+	}
+	GCsvProfilerInitialized = true;
+
 #if CSV_PROFILER_ALLOW_DEBUG_FEATURES
 	FParse::Value(FCommandLine::Get(), TEXT("csvCaptureOnEventFrameCount="), CaptureOnEventFrameCount);
 
@@ -4934,6 +4958,8 @@ void FCsvProfiler::Init()
 		BeginCapture(NumCsvFrames, FString(), FString(), ECsvProfilerFlags::BootCapture);
 		// Call BeginFrame() to start capturing a dummy first "frame"
 		BeginFrame();
+		// Add a region end marker so we can see where preinit and csv capture started
+		EndRegion(CSV_CATEGORY_INDEX_GLOBAL, TEXT("EngineInit/Launch"));
 	}
 	FParse::Value(FCommandLine::Get(), TEXT("csvRepeat="), GCsvRepeatCount);
 
diff --git a/Engine/Source/Runtime/Core/Public/ProfilingDebugging/CsvProfiler.h b/Engine/Source/Runtime/Core/Public/ProfilingDebugging/CsvProfiler.h
--- a/Engine/Source/Runtime/Core/Public/ProfilingDebugging/CsvProfiler.h
+++ b/Engine/Source/Runtime/Core/Public/ProfilingDebugging/CsvProfiler.h
@@ -413,6 +413,8 @@ public:
 
 	CORE_API void Init();
 
+	CORE_API static void PreInit(const TCHAR* Commandline);
+
 	/** Begin static interface (used by macros)*/
 	/** Push/pop events */
 	CORE_API static void BeginStat(const char* StatName, uint32 CategoryIndex, const char* NamedEventName = nullptr);
diff --git a/Engine/Source/Runtime/Launch/Private/LaunchEngineLoop.cpp b/Engine/Source/Runtime/Launch/Private/LaunchEngineLoop.cpp
--- a/Engine/Source/Runtime/Launch/Private/LaunchEngineLoop.cpp
+++ b/Engine/Source/Runtime/Launch/Private/LaunchEngineLoop.cpp
@@ -1732,6 +1732,10 @@ int32 FEngineLoop::PreInitPreStartupScreen(const TCHAR* CmdLine)
 {
 	ON_SCOPE_EXIT { FBootProfiling::OnPreInitPreStartupScreenComplete(); };
 	FDelayedAutoRegisterHelper::RunAndClearDelayedAutoRegisterDelegates(EDelayedRegisterRunPhase::StartOfEnginePreInit);
+#if WITH_ENGINE && CSV_PROFILER
+	// PreInit the CSV profiler as early as possible in case we're boot profiling with -csvCaptureFrames
+	FCsvProfiler::PreInit(CmdLine);
+#endif
 	SCOPED_BOOT_TIMING("FEngineLoop::PreInitPreStartupScreen");
 
 	// GLog is initialized lazily and its default primary thread is the thread that initialized it.
