Added asserts to FName to prevent them being created with invalid numbers.

#Synced-CL 50977583

diff --git a/Engine/Source/Runtime/Core/Private/UObject/UnrealNames.cpp b/Engine/Source/Runtime/Core/Private/UObject/UnrealNames.cpp
--- a/Engine/Source/Runtime/Core/Private/UObject/UnrealNames.cpp
+++ b/Engine/Source/Runtime/Core/Private/UObject/UnrealNames.cpp
@@ -2907,6 +2907,14 @@ bool FName::IsWithinBounds(FNameEntryId Id)
 	return GetNamePoolPostInit().IsValid(Id);
 }
 
+#if DO_CHECK
+void FName::ReportInvalidInternalNumber(FNameEntryId ComparisonId, int32 InternalNumber)
+{
+	UE_LOG(LogCore, Fatal, TEXT("Trying to construct an invalid FName '%s' with internal number: %d"), *GetNamePool().Resolve(ComparisonId).GetPlainNameString(), InternalNumber);
+	for (;;);
+}
+#endif
+
 /*-----------------------------------------------------------------------------
 	FName implementation.
 -----------------------------------------------------------------------------*/
@@ -3200,6 +3208,14 @@ private:
 			}
 		}
 #endif
+
+#if DO_CHECK
+		if (!FName::IsValidInternalNumber(InternalNumber))
+		{
+			FName::ReportInvalidInternalNumber(Ids.ComparisonId, InternalNumber);
+		}
+#endif
+
 		return MakeWithNumber(Ids, FindType, InternalNumber);
 	}
 
diff --git a/Engine/Source/Runtime/Core/Public/UObject/NameTypes.h b/Engine/Source/Runtime/Core/Public/UObject/NameTypes.h
--- a/Engine/Source/Runtime/Core/Public/UObject/NameTypes.h
+++ b/Engine/Source/Runtime/Core/Public/UObject/NameTypes.h
@@ -844,6 +844,17 @@ public:
 	/** Paranoid sanity check, same as IsValid() */
 	[[nodiscard]] bool IsValidIndexFast() const { return IsValid(); }
 
+	/**
+	 * Checks to see if an internal number is valid, such as those provided to the FName(Name, Number) constructors.
+	 *
+	 * @param	InInternalNumber	The number to check.
+	 *
+	 * @return	true if the number is valid.
+	 */
+	[[nodiscard]] static constexpr bool IsValidInternalNumber(int32 InInternalNumber)
+	{
+		return InInternalNumber >= 0;
+	}
 
 	/**
 	 * Checks to see that a given name-like string follows the rules that Unreal requires.
@@ -1303,6 +1314,14 @@ private:
 			// We need to store a new entry in the name table
 			return CreateNumberedName(ComparisonId, DisplayId, Number);
 		}
+
+#if DO_CHECK
+		if (!IsValidInternalNumber(Number))
+		{
+			ReportInvalidInternalNumber(ComparisonId, Number);
+		}
+#endif
+
 		// Otherwise we can just set the index members
 #endif
 		FName Out;
@@ -1316,7 +1335,7 @@ private:
 		return Out;
 	}
 
-	[[nodiscard]] FORCEINLINE static FName CreateNumberedNameIfNecessary(FNameEntryId ComparisonId, int32 Number)
+	[[nodiscard]] UE_REWRITE static FName CreateNumberedNameIfNecessary(FNameEntryId ComparisonId, int32 Number)
 	{
 		return CreateNumberedNameIfNecessary(ComparisonId, ComparisonId, Number);
 	}
@@ -1372,6 +1391,10 @@ private:
 	}
 #endif
 
+#if DO_CHECK
+	[[noreturn]] CORE_API static void ReportInvalidInternalNumber(FNameEntryId ComparisonId, int32 InternalNumber);
+#endif
+
 public:
 	template <typename T>
 	[[nodiscard]] FORCEINLINE auto UEOpEquals(T Rhs) const -> decltype(FName::Equals(*this, Rhs))
