SharedString: Make transactionally safe by using TTransactionalAtomicRefCount

- Changed TTransactionalAtomicRefCount::AddRef() memory order to relaxed and loosened other seq_cst usage.
- Moved TTransactionalAtomicRefCount into its own header to avoid bloating header dependencies.
- Added a TTransactionalAtomicRefCount constructor that takes the initial ref count to avoid an unnecessary and expensive AddRef() call.
- Added UE_NO_UNIQUE_ADDRESS to TransactionallyDeferredReleaseCount to allow it to have no size with ERefCountIsQueryable::No.
- Changed DeleteFn for TTransactionalAtomicRefCount to take a non-const pointer to support use with FMemory::Free.
- Fixed the ERefCountIsQueryable::No case which didn't compile before.
- Added tests from Neil.Henning of using FSharedString in a transaction.

#Synced-CL 50966494
#Synced-CL 50966494

diff --git a/Engine/Source/Programs/AutoRTFMTests/Private/UECoreTests.cpp b/Engine/Source/Programs/AutoRTFMTests/Private/UECoreTests.cpp
--- a/Engine/Source/Programs/AutoRTFMTests/Private/UECoreTests.cpp
+++ b/Engine/Source/Programs/AutoRTFMTests/Private/UECoreTests.cpp
@@ -32,6 +32,7 @@
 #include "UObject/UObjectGlobals.h"
 #include "MyAutoRTFMTestObject.h"
 #include "Containers/Queue.h"
+#include "Containers/SharedString.h"
 #include "Misc/CString.h"
 #include "Misc/ConfigCacheIni.h"
 #include "Misc/CoreDelegates.h"
@@ -2281,3 +2282,75 @@ TEST_CASE("UECore.TextTransformer")
 		});
 	}
 }
+
+TEST_CASE("UECore.FSharedString")
+{
+	SECTION("Create In Transaction")
+	{
+		AutoRTFM::Testing::Commit([&]
+		{
+			UE::FSharedString Str(TEXT("WOWWEE"));
+		});
+	}
+
+	SECTION("Create From StringView In Transaction")
+	{
+		TStringView<UE::FSharedString::ElementType> View(TEXT("WOWWEE"));
+
+		AutoRTFM::Testing::Commit([&]
+		{
+			UE::FSharedString Str(View);
+		});
+	}
+
+	SECTION("Copy Into Transaction")
+	{
+		UE::FSharedString Str(TEXT("WOWWEE"));
+
+		AutoRTFM::Testing::Commit([&]
+		{
+			UE::FSharedString Copy(Str);
+		});
+	}
+
+	SECTION("Assign In Transaction")
+	{
+		UE::FSharedString Str(TEXT("WOWWEE"));
+		UE::FSharedString Other;
+
+		AutoRTFM::Testing::Commit([&]
+		{
+			Other = Str;
+		});
+
+		REQUIRE(Str == Other);
+	}
+
+	SECTION("Assign Over Existing Data Transaction")
+	{
+		UE::FSharedString Str(TEXT("WOWWEE"));
+		UE::FSharedString Other(TEXT("NARP"));
+
+		AutoRTFM::Testing::Commit([&]
+		{
+			Other = Str;
+		});
+
+		REQUIRE(Str == Other);
+	}
+
+	SECTION("Assign By Move Over Existing Data Transaction")
+	{
+		UE::FSharedString Other(TEXT("NARP"));
+
+		AutoRTFM::Testing::Commit([&]
+		{
+			UE::FSharedString Str(TEXT("WOWWEE"));
+			Other = MoveTemp(Str);
+		});
+
+		UE::FSharedString Compare(TEXT("WOWWEE"));
+
+		REQUIRE(Compare == Other);
+	}
+}
diff --git a/Engine/Source/Runtime/Core/Private/AutoRTFM/TransactionalAtomicRefCount.cpp b/Engine/Source/Runtime/Core/Private/AutoRTFM/TransactionalAtomicRefCount.cpp
new file mode 100644
--- /dev/null
+++ b/Engine/Source/Runtime/Core/Private/AutoRTFM/TransactionalAtomicRefCount.cpp
@@ -0,0 +1,9 @@
+// Copyright Epic Games, Inc. All Rights Reserved.
+
+#include "AutoRTFM/TransactionalAtomicRefCount.h"
+#include "Misc/AssertionMacros.h"
+
+void UE::Private::CheckRefCountIsNonZero()
+{
+	check(!"Release() was called on an object which is already at a zero refcount.");
+}
diff --git a/Engine/Source/Runtime/Core/Private/Templates/RefCounting.cpp b/Engine/Source/Runtime/Core/Private/Templates/RefCounting.cpp
deleted file mode 100644
diff --git a/Engine/Source/Runtime/Core/Public/AutoRTFM/TransactionalAtomicRefCount.h b/Engine/Source/Runtime/Core/Public/AutoRTFM/TransactionalAtomicRefCount.h
new file mode 100644
--- /dev/null
+++ b/Engine/Source/Runtime/Core/Public/AutoRTFM/TransactionalAtomicRefCount.h
@@ -0,0 +1,235 @@
+// Copyright Epic Games, Inc. All Rights Reserved.
+
+#pragma once
+
+#include "AutoRTFM.h"
+#include "CoreTypes.h"
+
+#include <atomic>
+
+namespace UE::Private
+{
+
+CORE_API void CheckRefCountIsNonZero();
+
+enum class ERefCountIsQueryable : bool
+{
+	Yes = true,
+	No = false,
+};
+
+template <typename IntType, ERefCountIsQueryable Queryable> struct TQueryableTransactionalType;
+
+template <typename IntType> struct TQueryableTransactionalType<IntType, ERefCountIsQueryable::Yes>
+{
+	TQueryableTransactionalType(IntType Init)
+		: Payload(Init)
+	{
+	}
+
+	IntType operator++(int)
+	{
+		return Payload++;
+	}
+
+	IntType operator--(int)
+	{
+		return Payload--;
+	}
+
+	operator IntType() const
+	{
+		return Payload;
+	}
+private:
+	IntType Payload = 0;
+};
+
+template <typename IntType> struct TQueryableTransactionalType<IntType, ERefCountIsQueryable::No>
+{
+	TQueryableTransactionalType(IntType Init)
+	{
+	}
+
+	operator IntType() const
+	{
+		return 0;
+	}
+};
+
+/**
+ * TTransactionalAtomicRefCount manages a transactionally-safe atomic refcount value.
+ *
+ * This is used by FRefCountedObject and TRefCountingMixin (in thread-safe mode).
+ */
+template <typename IntType, ERefCountIsQueryable Queryable = ERefCountIsQueryable::Yes>
+class TTransactionalAtomicRefCount
+{
+public:
+	TTransactionalAtomicRefCount() = default;
+
+	explicit TTransactionalAtomicRefCount(IntType InitialRefCount)
+		: RefCount(InitialRefCount)
+	{
+	}
+
+	template <auto DeleteFn>
+	IntType AddRef() const
+	{
+		if (!AutoRTFM::IsClosed())
+		{
+			return ++RefCount;
+		}
+
+		IntType Refs;
+
+		UE_AUTORTFM_OPEN
+		{
+			Refs = RefCount.fetch_add(1, std::memory_order_relaxed);
+		};
+
+		// As we are inside a transaction, the rollback must undo our refcount bump.
+		// In general, this is best handled by Release(). However, there is one case 
+		// that needs to be handled with special care. A brand-new object has a
+		// refcount of zero, and a rollback must return it to this zero-refcount state.
+		// However, calling AddRef() followed by Release() would not accomplish this;
+		// instead, it would free the object entirely! We need to guard against this,
+		// since it could lead to a double-free, so we detect the zero-reference case
+		// and special-case it.
+		if (Refs == 0)
+		{
+			UE_AUTORTFM_ONABORT(this)
+			{
+				RefCount.fetch_sub(1, std::memory_order_release);
+				// The refcount is likely zero now, but leaving the object alive isn't a leak.
+				// We are restoring the object back to its initial "just-created" state.
+			};
+		}
+		else
+		{
+			UE_AUTORTFM_ONABORT(this)
+			{
+				Release<DeleteFn>();
+			};
+		}
+
+		return Refs + 1 - TransactionallyDeferredReleaseCount;
+	}
+
+	template <auto DeleteFn>
+	IntType Release() const
+	{
+	#if DO_GUARD_SLOW
+		if (RefCount.load(std::memory_order_relaxed) == 0)
+		{
+			CheckRefCountIsNonZero();
+		}
+	#endif
+
+		if (AutoRTFM::IsClosed())
+		{
+			IntType Refs;
+
+			UE_AUTORTFM_OPEN
+			{
+				if constexpr (Queryable == ERefCountIsQueryable::Yes)
+				{
+					TransactionallyDeferredReleaseCount++;
+				}
+				Refs = RefCount.load(std::memory_order_relaxed) - TransactionallyDeferredReleaseCount;
+			};
+
+			if constexpr (Queryable == ERefCountIsQueryable::Yes)
+			{
+				UE_AUTORTFM_ONABORT(this)
+				{
+					// TODO: degrade this to a `checkSlow` once time has passed.
+				#if DO_CHECK
+					if (TransactionallyDeferredReleaseCount == 0)
+					{
+						CheckRefCountIsNonZero();
+					}
+				#endif
+					TransactionallyDeferredReleaseCount--;
+				};
+			}
+
+			// Refcount changes and frees are deferred until the transaction is concluded.
+			UE_AUTORTFM_ONCOMMIT(this)
+			{
+				if constexpr (Queryable == ERefCountIsQueryable::Yes)
+				{
+					// TODO (SOL-7350): degrade this to a `checkSlow` once time has passed.
+				#if DO_CHECK
+					if (TransactionallyDeferredReleaseCount == 0)
+					{
+						CheckRefCountIsNonZero();
+					}
+				#endif
+					TransactionallyDeferredReleaseCount--;
+				}
+				ImmediatelyRelease<DeleteFn>();
+			};
+
+			return Refs;
+		}
+		else
+		{
+			return ImmediatelyRelease<DeleteFn>() - 1;
+		}
+	}
+
+	// This is equivalent to https://en.cppreference.com/w/cpp/memory/shared_ptr/use_count
+	// This loads with a relaxed memory ordering because a 'live' reference count is
+	// unstable by nature and so there's no benefit to try and enforce memory ordering
+	// around the reading of it.
+	IntType GetRefCount() const
+		requires (Queryable == ERefCountIsQueryable::Yes)
+	{
+		if (AutoRTFM::IsClosed())
+		{
+			IntType Refs;
+			UE_AUTORTFM_OPEN
+			{
+				Refs = RefCount.load(std::memory_order_relaxed) - TransactionallyDeferredReleaseCount;
+			};
+			return Refs;
+		}
+		else
+		{
+			return RefCount.load(std::memory_order_relaxed);
+		}
+	}
+
+private:
+	template <auto DeleteFn>
+	IntType ImmediatelyRelease() const
+	{
+		// fetch_sub returns the refcount _before_ it was decremented. std::memory_order_acq_rel is
+		// used so that, if we do end up executing the destructor, it's not possible for side effects 
+		// from executing the destructor to end up being visible before we've determined that the 
+		// reference count is actually zero.
+		IntType RefsBeforeRelease = RefCount.fetch_sub(1, std::memory_order_acq_rel);
+
+	#if DO_GUARD_SLOW
+		// A check-failure is issued if an object is over-released.
+		if (RefsBeforeRelease == 0)
+		{
+			CheckRefCountIsNonZero();
+		}
+	#endif
+
+		// We immediately free the object if its refcount has become zero.
+		if (RefsBeforeRelease == 1)
+		{
+			DeleteFn(const_cast<TTransactionalAtomicRefCount*>(this));
+		}
+
+		return RefsBeforeRelease;
+	}
+
+	mutable std::atomic<IntType> RefCount = 0;
+	UE_NO_UNIQUE_ADDRESS mutable TQueryableTransactionalType<IntType, Queryable> TransactionallyDeferredReleaseCount = 0;
+};
+
+} // UE::Private
diff --git a/Engine/Source/Runtime/Core/Public/Containers/SharedString.h b/Engine/Source/Runtime/Core/Public/Containers/SharedString.h
--- a/Engine/Source/Runtime/Core/Public/Containers/SharedString.h
+++ b/Engine/Source/Runtime/Core/Public/Containers/SharedString.h
@@ -2,6 +2,7 @@
 
 #pragma once
 
+#include "AutoRTFM/TransactionalAtomicRefCount.h"
 #include "Containers/StringFwd.h"
 #include "Containers/StringView.h"
 #include "CoreTypes.h"
@@ -107,6 +108,9 @@ private:
 	CharType* Chars = nullptr;
 
 	static constexpr inline CharType NullChar{};
+
+	using FRefCountType = Private::TTransactionalAtomicRefCount<int32, Private::ERefCountIsQueryable::No>;
+	static_assert(sizeof(FRefCountType) == sizeof(int32));
 };
 
 template <typename CharType>
@@ -156,10 +160,10 @@ template <typename CharType>
 	{
 		static_assert(alignof(int32) <= MIN_ALIGNMENT);
 		static_assert(alignof(int32) >= alignof(CharType));
-		static_assert(sizeof(int32) == sizeof(std::atomic<int32>));
 		const SIZE_T Size = sizeof(int32) + sizeof(int32) + sizeof(CharType) + sizeof(CharType) * Length;
 		int32* const Header = static_cast<int32*>(FMemory::Malloc(Size));
-		new((void*)&Header[0]) std::atomic<int32>(1); // Ref Count
+		// RefCount needs to be at the same address as the allocation to make FMemory::Free work as the DeleteFn.
+		new((void*)&Header[0]) FRefCountType(1);
 		new((void*)&Header[1]) int32(Length);
 		Chars = reinterpret_cast<CharType*>(&Header[2]);
 		String.CopyString(Chars, Length);
@@ -180,8 +184,8 @@ void TSharedString<CharType>::AddRef(CharType* const MaybeNullChars)
 	{
 		static_assert(sizeof(int32) == sizeof(std::atomic<int32>));
 		int32* const Header = reinterpret_cast<int32*>(MaybeNullChars) - 2;
-		std::atomic<int32>& RefCount = reinterpret_cast<std::atomic<int32>&>(Header[0]);
-		RefCount.fetch_add(1, std::memory_order_relaxed);
+		FRefCountType& RefCount = reinterpret_cast<FRefCountType&>(Header[0]);
+		RefCount.AddRef<FMemory::Free>();
 	}
 }
 
@@ -192,11 +196,8 @@ void TSharedString<CharType>::Release(CharType* const MaybeNullChars)
 	{
 		static_assert(sizeof(int32) == sizeof(std::atomic<int32>));
 		int32* const Header = reinterpret_cast<int32*>(MaybeNullChars) - 2;
-		std::atomic<int32>& RefCount = reinterpret_cast<std::atomic<int32>&>(Header[0]);
-		if (RefCount.fetch_sub(1, std::memory_order_acq_rel) == 1)
-		{
-			FMemory::Free(Header);
-		}
+		FRefCountType& RefCount = reinterpret_cast<FRefCountType&>(Header[0]);
+		RefCount.Release<FMemory::Free>();
 	}
 }
 
diff --git a/Engine/Source/Runtime/Core/Public/Templates/RefCounting.h b/Engine/Source/Runtime/Core/Public/Templates/RefCounting.h
--- a/Engine/Source/Runtime/Core/Public/Templates/RefCounting.h
+++ b/Engine/Source/Runtime/Core/Public/Templates/RefCounting.h
@@ -4,6 +4,7 @@
 
 #include "CoreTypes.h"
 #include "AutoRTFM.h"
+#include "AutoRTFM/TransactionalAtomicRefCount.h"
 #include "HAL/PlatformAtomics.h"
 #include "HAL/PreprocessorHelpers.h"
 #include "HAL/ThreadSafeCounter.h"
@@ -28,7 +29,7 @@
 struct FReturnedRefCountValue
 {
 	explicit FReturnedRefCountValue(uint32 InRefCount)
-		: RefCount(InRefCount)	 
+		: RefCount(InRefCount)
 	{
 	}
 	FReturnedRefCountValue(const FReturnedRefCountValue& Other) = default;
@@ -56,198 +57,9 @@ private:
 	uint32 RefCount = 0;
 };
 
-/**
- * UE::Private::TTransactionalAtomicRefCount manages a transactionally-safe atomic refcount value.
- * This is used by FRefCountedObject and TRefCountingMixin (in thread-safe mode).
- */
 namespace UE::Private
 {
 
-CORE_API void CheckRefCountIsNonZero();
-
-enum class ERefCountIsQueryable : bool
-{
-	Yes = true,
-	No = false,
-};
-
-template <typename IntType, ERefCountIsQueryable Queryable> struct TQueryableTransactionalType;
-
-template <typename IntType> struct TQueryableTransactionalType<IntType, ERefCountIsQueryable::Yes>
-{
-	TQueryableTransactionalType(IntType Init) : Payload(Init) {}
-
-	IntType operator++(int)
-	{
-		return Payload++;
-	}
-
-	IntType operator--(int)
-	{
-		return Payload--;
-	}
-
-	operator IntType() const
-	{
-		return Payload;
-	}
-private:
-	IntType Payload = 0;
-};
-
-template <typename IntType> struct TQueryableTransactionalType<IntType, ERefCountIsQueryable::No>
-{
-	TQueryableTransactionalType(IntType Init) {}
-
-	operator IntType() const
-	{
-		return 0;
-	}
-};
-
-template <typename IntType, ERefCountIsQueryable Queryable = ERefCountIsQueryable::Yes>
-class TTransactionalAtomicRefCount
-{
-public:
-	template <auto DeleteFn>
-	IntType AddRef() const
-	{
-		if (!AutoRTFM::IsClosed())
-		{
-			return ++RefCount;
-		}
-
-		IntType Refs;
-
-		UE_AUTORTFM_OPEN
-		{
-			Refs = RefCount++;
-		};
-
-		// As we are inside a transaction, the rollback must undo our refcount bump.
-		// In general, this is best handled by Release(). However, there is one case 
-		// that needs to be handled with special care. A brand-new object has a
-		// refcount of zero, and a rollback must return it to this zero-refcount state.
-		// However, calling AddRef() followed by Release() would not accomplish this;
-		// instead, it would free the object entirely! We need to guard against this,
-		// since it could lead to a double-free, so we detect the zero-reference case
-		// and special-case it.
-		if (Refs == 0)
-		{
-			UE_AUTORTFM_ONABORT(this)
-			{
-				--RefCount;
-				// The refcount is likely zero now, but leaving the object alive isn't a leak.
-				// We are restoring the object back to its initial "just-created" state.
-			};
-		}
-		else
-		{
-			UE_AUTORTFM_ONABORT(this)
-			{
-				Release<DeleteFn>();
-			};
-		}
-
-		return Refs + 1 - TransactionallyDeferredReleaseCount;
-	}
-
-	template <auto DeleteFn>
-	IntType Release() const
-	{
-#if DO_GUARD_SLOW
-		if (RefCount == 0)
-		{
-			CheckRefCountIsNonZero();
-		}
-#endif
-
-		if (AutoRTFM::IsClosed())
-		{
-			IntType Refs;
-
-			UE_AUTORTFM_OPEN
-			{
-				TransactionallyDeferredReleaseCount++;
-				Refs = RefCount.load(std::memory_order_relaxed) - TransactionallyDeferredReleaseCount;
-			};
-
-			UE_AUTORTFM_ONABORT(this)
-			{
-				// TODO: degrade this to a `checkSlow` once time has passed.
-				check(TransactionallyDeferredReleaseCount > 0);
-				TransactionallyDeferredReleaseCount--;
-			};
-
-			// Refcount changes and frees are deferred until the transaction is concluded.
-			UE_AUTORTFM_ONCOMMIT(this)
-			{
-				// TODO (SOL-7350): degrade this to a `checkSlow` once time has passed.
-				check(TransactionallyDeferredReleaseCount > 0);
-				TransactionallyDeferredReleaseCount--;
-				ImmediatelyRelease<DeleteFn>();
-			};
-
-			return Refs;
-		}
-		else
-		{
-			return ImmediatelyRelease<DeleteFn>() - 1;
-		}
-	}
-
-	// This is equivalent to https://en.cppreference.com/w/cpp/memory/shared_ptr/use_count
-	// This loads with a relaxed memory ordering because a 'live' reference count is
-	// unstable by nature and so there's no benefit to try and enforce memory ordering
-	// around the reading of it.
-	IntType GetRefCount() const
-	{
-		if (AutoRTFM::IsClosed())
-		{
-			IntType Refs;
-			UE_AUTORTFM_OPEN
-			{
-				Refs = RefCount.load(std::memory_order_relaxed) - TransactionallyDeferredReleaseCount;
-			};
-			return Refs;
-		}
-		else
-		{
-			return RefCount.load(std::memory_order_relaxed);
-		}
-	}
-
-private:
-	template <auto DeleteFn>
-	IntType ImmediatelyRelease() const
-	{
-		// fetch_sub returns the refcount _before_ it was decremented. std::memory_order_acq_rel is
-		// used so that, if we do end up executing the destructor, it's not possible for side effects 
-		// from executing the destructor to end up being visible before we've determined that the 
-		// reference count is actually zero.
-		IntType RefsBeforeRelease = RefCount.fetch_sub(1, std::memory_order_acq_rel);
-
-#if DO_GUARD_SLOW
-		// A check-failure is issued if an object is over-released.
-		if (RefsBeforeRelease == 0)
-		{
-			CheckRefCountIsNonZero();
-		}
-#endif
-
-		// We immediately free the object if its refcount has become zero.
-		if (RefsBeforeRelease == 1)
-		{
-			DeleteFn(this);
-		}
-
-		return RefsBeforeRelease;
-	}
-
-	mutable std::atomic<IntType> RefCount = 0;
-	mutable TQueryableTransactionalType<IntType, Queryable> TransactionallyDeferredReleaseCount = 0;
-};
-
 /**
  * Please do not use this, long term we want it to go away. Being able to query the reference count can
  * easily lead into really gnarly bugs and we generally don't want this sort of paradigm to be used more
