Moved TransactionalAtomicRefCount to Async/AtomicRefCount.h

#rnx
#Synced-CL 50966494
#Synced-CL 50966494

diff --git a/Engine/Source/Runtime/Core/Private/Async/AtomicRefCount.cpp b/Engine/Source/Runtime/Core/Private/Async/AtomicRefCount.cpp
new file mode 100644
--- /dev/null
+++ b/Engine/Source/Runtime/Core/Private/Async/AtomicRefCount.cpp
@@ -0,0 +1,9 @@
+// Copyright Epic Games, Inc. All Rights Reserved.
+
+#include "Async/AtomicRefCount.h"
+#include "Misc/AssertionMacros.h"
+
+void UE::Private::CheckRefCountIsNonZero()
+{
+	check(!"Release() was called on an object which is already at a zero refcount.");
+}
diff --git a/Engine/Source/Runtime/Core/Private/AutoRTFM/TransactionalAtomicRefCount.cpp b/Engine/Source/Runtime/Core/Private/AutoRTFM/TransactionalAtomicRefCount.cpp
deleted file mode 100644
diff --git a/Engine/Source/Runtime/Core/Public/Async/AtomicRefCount.h b/Engine/Source/Runtime/Core/Public/Async/AtomicRefCount.h
new file mode 100644
--- /dev/null
+++ b/Engine/Source/Runtime/Core/Public/Async/AtomicRefCount.h
@@ -0,0 +1,233 @@
+// Copyright Epic Games, Inc. All Rights Reserved.
+
+#pragma once
+
+#include "AutoRTFM.h"
+#include "CoreTypes.h"
+
+#include <atomic>
+
+namespace UE::Private
+{
+
+CORE_API void CheckRefCountIsNonZero();
+
+enum class ERefCountIsQueryable : bool
+{
+	Yes = true,
+	No = false,
+};
+
+template <typename IntType, ERefCountIsQueryable Queryable> struct TQueryableTransactionalType;
+
+template <typename IntType> struct TQueryableTransactionalType<IntType, ERefCountIsQueryable::Yes>
+{
+	TQueryableTransactionalType(IntType Init)
+		: Payload(Init)
+	{
+	}
+
+	IntType operator++(int)
+	{
+		return Payload++;
+	}
+
+	IntType operator--(int)
+	{
+		return Payload--;
+	}
+
+	operator IntType() const
+	{
+		return Payload;
+	}
+private:
+	IntType Payload = 0;
+};
+
+template <typename IntType> struct TQueryableTransactionalType<IntType, ERefCountIsQueryable::No>
+{
+	TQueryableTransactionalType(IntType Init)
+	{
+	}
+
+	operator IntType() const
+	{
+		return 0;
+	}
+};
+
+/**
+ * TAtomicRefCount manages a transactionally-safe atomic refcount value.
+ *
+ * This is used by FRefCountedObject and TRefCountingMixin (in thread-safe mode).
+ */
+template <typename IntType, ERefCountIsQueryable Queryable = ERefCountIsQueryable::Yes>
+class TAtomicRefCount
+{
+public:
+	TAtomicRefCount() = default;
+
+	explicit TAtomicRefCount(IntType InitialRefCount)
+		: RefCount(InitialRefCount)
+	{
+	}
+
+	template <auto DeleteFn>
+	IntType AddRef() const
+	{
+		if (!AutoRTFM::IsClosed())
+		{
+			return ++RefCount;
+		}
+
+		IntType Refs;
+
+		UE_AUTORTFM_OPEN
+		{
+			Refs = RefCount.fetch_add(1, std::memory_order_relaxed);
+		};
+
+		// As we are inside a transaction, the rollback must undo our refcount bump.
+		// In general, this is best handled by Release(). However, there is one case 
+		// that needs to be handled with special care. A brand-new object has a
+		// refcount of zero, and a rollback must return it to this zero-refcount state.
+		// However, calling AddRef() followed by Release() would not accomplish this;
+		// instead, it would free the object entirely! We need to guard against this,
+		// since it could lead to a double-free, so we detect the zero-reference case
+		// and special-case it.
+		if (Refs == 0)
+		{
+			UE_AUTORTFM_ONABORT(this)
+			{
+				RefCount.fetch_sub(1, std::memory_order_release);
+				// The refcount is likely zero now, but leaving the object alive isn't a leak.
+				// We are restoring the object back to its initial "just-created" state.
+			};
+		}
+		else
+		{
+			UE_AUTORTFM_ONABORT(this)
+			{
+				Release<DeleteFn>();
+			};
+		}
+
+		return Refs + 1 - TransactionallyDeferredReleaseCount;
+	}
+
+	template <auto DeleteFn>
+	IntType Release() const
+	{
+	#if DO_GUARD_SLOW
+		if (RefCount.load(std::memory_order_relaxed) == 0)
+		{
+			CheckRefCountIsNonZero();
+		}
+	#endif
+
+		if (AutoRTFM::IsClosed())
+		{
+			IntType Refs;
+
+			UE_AUTORTFM_OPEN
+			{
+				if constexpr (Queryable == ERefCountIsQueryable::Yes)
+				{
+					TransactionallyDeferredReleaseCount++;
+				}
+				Refs = RefCount.load(std::memory_order_relaxed) - TransactionallyDeferredReleaseCount;
+			};
+
+			if constexpr (Queryable == ERefCountIsQueryable::Yes)
+			{
+				UE_AUTORTFM_ONABORT(this)
+				{
+				#if DO_GUARD_SLOW
+					if (TransactionallyDeferredReleaseCount == 0)
+					{
+						CheckRefCountIsNonZero();
+					}
+				#endif
+					TransactionallyDeferredReleaseCount--;
+				};
+			}
+
+			// Refcount changes and frees are deferred until the transaction is concluded.
+			UE_AUTORTFM_ONCOMMIT(this)
+			{
+				if constexpr (Queryable == ERefCountIsQueryable::Yes)
+				{
+				#if DO_GUARD_SLOW
+					if (TransactionallyDeferredReleaseCount == 0)
+					{
+						CheckRefCountIsNonZero();
+					}
+				#endif
+					TransactionallyDeferredReleaseCount--;
+				}
+				ImmediatelyRelease<DeleteFn>();
+			};
+
+			return Refs;
+		}
+		else
+		{
+			return ImmediatelyRelease<DeleteFn>() - 1;
+		}
+	}
+
+	// This is equivalent to https://en.cppreference.com/w/cpp/memory/shared_ptr/use_count
+	// This loads with a relaxed memory ordering because a 'live' reference count is
+	// unstable by nature and so there's no benefit to try and enforce memory ordering
+	// around the reading of it.
+	IntType GetRefCount() const
+		requires (Queryable == ERefCountIsQueryable::Yes)
+	{
+		if (AutoRTFM::IsClosed())
+		{
+			IntType Refs;
+			UE_AUTORTFM_OPEN
+			{
+				Refs = RefCount.load(std::memory_order_relaxed) - TransactionallyDeferredReleaseCount;
+			};
+			return Refs;
+		}
+		else
+		{
+			return RefCount.load(std::memory_order_relaxed);
+		}
+	}
+
+private:
+	template <auto DeleteFn>
+	IntType ImmediatelyRelease() const
+	{
+		// fetch_sub returns the refcount _before_ it was decremented. std::memory_order_acq_rel is
+		// used so that, if we do end up executing the destructor, it's not possible for side effects 
+		// from executing the destructor to end up being visible before we've determined that the 
+		// reference count is actually zero.
+		IntType RefsBeforeRelease = RefCount.fetch_sub(1, std::memory_order_acq_rel);
+
+	#if DO_GUARD_SLOW
+		// A check-failure is issued if an object is over-released.
+		if (RefsBeforeRelease == 0)
+		{
+			CheckRefCountIsNonZero();
+		}
+	#endif
+
+		// We immediately free the object if its refcount has become zero.
+		if (RefsBeforeRelease == 1)
+		{
+			DeleteFn(const_cast<TAtomicRefCount*>(this));
+		}
+
+		return RefsBeforeRelease;
+	}
+
+	mutable std::atomic<IntType> RefCount = 0;
+	UE_NO_UNIQUE_ADDRESS mutable TQueryableTransactionalType<IntType, Queryable> TransactionallyDeferredReleaseCount = 0;
+};
+
+} // UE::Private
diff --git a/Engine/Source/Runtime/Core/Public/AutoRTFM/TransactionalAtomicRefCount.h b/Engine/Source/Runtime/Core/Public/AutoRTFM/TransactionalAtomicRefCount.h
deleted file mode 100644
diff --git a/Engine/Source/Runtime/Core/Public/Containers/SharedString.h b/Engine/Source/Runtime/Core/Public/Containers/SharedString.h
--- a/Engine/Source/Runtime/Core/Public/Containers/SharedString.h
+++ b/Engine/Source/Runtime/Core/Public/Containers/SharedString.h
@@ -2,7 +2,7 @@
 
 #pragma once
 
-#include "AutoRTFM/TransactionalAtomicRefCount.h"
+#include "Async/AtomicRefCount.h"
 #include "Containers/StringFwd.h"
 #include "Containers/StringView.h"
 #include "CoreTypes.h"
@@ -109,7 +109,7 @@ private:
 
 	static constexpr inline CharType NullChar{};
 
-	using FRefCountType = Private::TTransactionalAtomicRefCount<int32, Private::ERefCountIsQueryable::No>;
+	using FRefCountType = Private::TAtomicRefCount<int32, Private::ERefCountIsQueryable::No>;
 	static_assert(sizeof(FRefCountType) == sizeof(int32));
 };
 
diff --git a/Engine/Source/Runtime/Core/Public/Templates/RefCounting.h b/Engine/Source/Runtime/Core/Public/Templates/RefCounting.h
--- a/Engine/Source/Runtime/Core/Public/Templates/RefCounting.h
+++ b/Engine/Source/Runtime/Core/Public/Templates/RefCounting.h
@@ -3,8 +3,8 @@
 #pragma once
 
 #include "CoreTypes.h"
+#include "Async/AtomicRefCount.h"
 #include "AutoRTFM.h"
-#include "AutoRTFM/TransactionalAtomicRefCount.h"
 #include "HAL/PlatformAtomics.h"
 #include "HAL/PreprocessorHelpers.h"
 #include "HAL/ThreadSafeCounter.h"
@@ -66,7 +66,7 @@ namespace UE::Private
  * widely. But there are a few places where it is currently being legitimately used so we'll provide a
  * queryable variant to let us deprecate the return of the reference count from `FRefCountedObject` asap.
  */
-class FQueryableRefCountedObject : UE::Private::TTransactionalAtomicRefCount<uint32, UE::Private::ERefCountIsQueryable::Yes>
+class FQueryableRefCountedObject : UE::Private::TAtomicRefCount<uint32, UE::Private::ERefCountIsQueryable::Yes>
 {
 public:
 	FQueryableRefCountedObject() = default;
@@ -91,7 +91,7 @@ public:
 	}
 
 private:
-	using Super = UE::Private::TTransactionalAtomicRefCount<uint32, UE::Private::ERefCountIsQueryable::Yes>;
+	using Super = UE::Private::TAtomicRefCount<uint32, UE::Private::ERefCountIsQueryable::Yes>;
 
 	static void DeleteThis(const Super* This)
 	{
@@ -115,7 +115,7 @@ public:
 /**
  * Base class implementing thread-safe reference counting.
  */
-class FRefCountedObject : UE::Private::TTransactionalAtomicRefCount<uint32>
+class FRefCountedObject : UE::Private::TAtomicRefCount<uint32>
 {
 public:
 	FRefCountedObject() = default;
@@ -140,7 +140,7 @@ public:
 	}
 
 private:
-	using Super = UE::Private::TTransactionalAtomicRefCount<uint32>;
+	using Super = UE::Private::TAtomicRefCount<uint32>;
 
 	static void DeleteThis(const Super* This)
 	{
@@ -194,7 +194,7 @@ class TRefCountingMixin;
  * Thread-safe specialization
  */
 template <typename T>
-class TRefCountingMixin<T, ERefCountingMode::ThreadSafe> : UE::Private::TTransactionalAtomicRefCount<uint32>
+class TRefCountingMixin<T, ERefCountingMode::ThreadSafe> : UE::Private::TAtomicRefCount<uint32>
 {
 public:
 	TRefCountingMixin() = default;
@@ -223,12 +223,12 @@ public:
 	}
 
 private:
-	using Super = UE::Private::TTransactionalAtomicRefCount<uint32>;
+	using Super = UE::Private::TAtomicRefCount<uint32>;
 
 	static void StaticDestroyMixin(const Super* This)
 	{
 		// This static_cast is traversing two levels of the class hierarchy.
-		// We are casting from our parent class (TTransactionalAtomicRefCount*) to our subclass (T*).
+		// We are casting from our parent class (TAtomicRefCount*) to our subclass (T*).
 		T::StaticDestroyObject(static_cast<const T*>(This));
 	}
 };
