Suppress false positive from PVS

#rnx
#Synced-CL 50967057

diff --git a/Engine/Source/Runtime/Core/Public/Async/AtomicRefCount.h b/Engine/Source/Runtime/Core/Public/Async/AtomicRefCount.h
--- a/Engine/Source/Runtime/Core/Public/Async/AtomicRefCount.h
+++ b/Engine/Source/Runtime/Core/Public/Async/AtomicRefCount.h
@@ -220,7 +220,7 @@ private:
 		// We immediately free the object if its refcount has become zero.
 		if (RefsBeforeRelease == 1)
 		{
-			DeleteFn(const_cast<TAtomicRefCount*>(this));
+			DeleteFn(const_cast<TAtomicRefCount*>(this)); //-V607
 		}
 
 		return RefsBeforeRelease;
