Changed FName's Number parameters to InternalNumber to better reflect reality.

#Synced-CL 50977583

diff --git a/Engine/Source/Runtime/Core/Private/UObject/UnrealNames.cpp b/Engine/Source/Runtime/Core/Private/UObject/UnrealNames.cpp
--- a/Engine/Source/Runtime/Core/Private/UObject/UnrealNames.cpp
+++ b/Engine/Source/Runtime/Core/Private/UObject/UnrealNames.cpp
@@ -3015,10 +3015,10 @@ struct FNameHelper
 		}
 		
 		uint32 InternalNumber = UE::Core::Private::ParseNumberFromName(View.Str, /* may be shortened */ View.Len);
-		return MakeWithNumber(View, FindType, InternalNumber);
+		return MakeWithInternalNumber(View, FindType, InternalNumber);
 	}
 
-	static FName MakeWithNumber(FNameAnsiStringView	 View, EFindName FindType, int32 InternalNumber)
+	static FName MakeWithInternalNumber(FNameAnsiStringView	 View, EFindName FindType, int32 InternalNumber)
 	{
 		// Ignore the supplied number if the name string is empty
 		// to keep the semantics of the old FName implementation
@@ -3030,7 +3030,7 @@ struct FNameHelper
 		return MakeInternal(FNameStringView(View.Str, View.Len), FindType, InternalNumber);
 	}
 
-	static FName MakeWithNumber(FNameUtf8StringViewWithWidth View, EFindName FindType, int32 InternalNumber)
+	static FName MakeWithInternalNumber(FNameUtf8StringViewWithWidth View, EFindName FindType, int32 InternalNumber)
 	{
 		// Ignore the supplied number if the name string is empty
 		// to keep the semantics of the old FName implementation
@@ -3050,7 +3050,7 @@ struct FNameHelper
 		}
 	}
 
-	static FName MakeWithNumber(const FWideStringViewWithWidth View, EFindName FindType, int32 InternalNumber)
+	static FName MakeWithInternalNumber(const FWideStringViewWithWidth View, EFindName FindType, int32 InternalNumber)
 	{
 		// Ignore the supplied number if the name string is empty
 		// to keep the semantics of the old FName implementation
@@ -3085,7 +3085,7 @@ struct FNameHelper
 		return MakeInternal(View, FNAME_Add, NAME_NO_NUMBER_INTERNAL);
 	}
 
-	static FName MakeWithNumber(FNameEntryIds BaseIds, EFindName FindType, int32 InternalNumber)
+	static FName MakeWithInternalNumber(FNameEntryIds BaseIds, EFindName FindType, int32 InternalNumber)
 	{
 #if UE_FNAME_OUTLINE_NUMBER
 #if WITH_CASE_PRESERVING_NAME
@@ -3192,7 +3192,7 @@ private:
 			// We need to disambiguate here between "we found 'None'" and "we didn't find what we were looking for"
 			if (View.IsNoneString())
 			{
-				return MakeWithNumber(Ids, FindType, InternalNumber); // 'None' with suffix
+				return MakeWithInternalNumber(Ids, FindType, InternalNumber); // 'None' with suffix
 			}
 			else
 			{
@@ -3200,7 +3200,7 @@ private:
 			}
 		}
 #endif
-		return MakeWithNumber(Ids, FindType, InternalNumber);
+		return MakeWithInternalNumber(Ids, FindType, InternalNumber);
 	}
 
 	// Indices have already been numbered if necessary
@@ -3339,9 +3339,9 @@ FNameEntryId FName::GetDisplayIndex() const
 #endif
 }
 
-FName FName::FindNumberedName(FNameEntryId DisplayId, int32 Number)
+FName FName::FindNumberedName(FNameEntryId DisplayId, int32 InInternalNumber)
 {
-	return FNameHelper::MakeWithNumber(FNameEntryIds{FName::GetComparisonIdFromDisplayId(DisplayId),DisplayId}, FNAME_Find , Number);
+	return FNameHelper::MakeWithInternalNumber(FNameEntryIds{FName::GetComparisonIdFromDisplayId(DisplayId), DisplayId}, FNAME_Find, InInternalNumber);
 }
 
 int32 FName::GetNumber() const 
@@ -3358,9 +3358,9 @@ int32 FName::GetNumber() const
 		return NAME_NO_NUMBER_INTERNAL;
 	}
 }
-void FName::SetNumber(int32 InNumber) 
+void FName::SetNumber(int32 InInternalNumber) 
 {
-	FName NewName = FNameHelper::MakeWithNumber(FNameEntryIds{GetComparisonIndex(), GetDisplayIndex()}, FNAME_Add, InNumber);
+	FName NewName = FNameHelper::MakeWithInternalNumber(FNameEntryIds{GetComparisonIndex(), GetDisplayIndex()}, FNAME_Add, InInternalNumber);
 	*this = NewName;
 }
 #endif // UE_FNAME_OUTLINE_NUMBER
@@ -3424,37 +3424,37 @@ FName::FName(int32 Len, const UTF8CHAR* Name, EFindName FindType)
 	: FName(FNameHelper::MakeDetectNumber(MakeUnconvertedView(Name, Len), FindType))
 {}
 
-FName::FName(const WIDECHAR* Name, int32 InNumber)
-	: FName(FNameHelper::MakeWithNumber(MakeUnconvertedView(Name), FNAME_Add, InNumber))
+FName::FName(const WIDECHAR* Name, int32 InInternalNumber)
+	: FName(FNameHelper::MakeWithInternalNumber(MakeUnconvertedView(Name), FNAME_Add, InInternalNumber))
 {}
 
-FName::FName(const ANSICHAR* Name, int32 InNumber)
-	: FName(FNameHelper::MakeWithNumber(MakeUnconvertedView(Name), FNAME_Add, InNumber))
+FName::FName(const ANSICHAR* Name, int32 InInternalNumber)
+	: FName(FNameHelper::MakeWithInternalNumber(MakeUnconvertedView(Name), FNAME_Add, InInternalNumber))
 {}
 
-FName::FName(const UTF8CHAR* Name, int32 InNumber)
-	: FName(FNameHelper::MakeWithNumber(MakeUnconvertedView(Name), FNAME_Add, InNumber))
+FName::FName(const UTF8CHAR* Name, int32 InInternalNumber)
+	: FName(FNameHelper::MakeWithInternalNumber(MakeUnconvertedView(Name), FNAME_Add, InInternalNumber))
 {}
 
-FName::FName(int32 Len, const WIDECHAR* Name, int32 InNumber)
-	: FName(InNumber != NAME_NO_NUMBER_INTERNAL ? FNameHelper::MakeWithNumber(MakeUnconvertedView(Name, Len), FNAME_Add, InNumber)
+FName::FName(int32 Len, const WIDECHAR* Name, int32 InInternalNumber)
+	: FName(InInternalNumber != NAME_NO_NUMBER_INTERNAL ? FNameHelper::MakeWithInternalNumber(MakeUnconvertedView(Name, Len), FNAME_Add, InInternalNumber)
 		: FNameHelper::MakeDetectNumber(MakeUnconvertedView(Name, Len), FNAME_Add))
 {}
 
-FName::FName(int32 Len, const ANSICHAR* Name, int32 InNumber)
-	: FName(InNumber != NAME_NO_NUMBER_INTERNAL ? FNameHelper::MakeWithNumber(MakeUnconvertedView(Name, Len), FNAME_Add, InNumber)
+FName::FName(int32 Len, const ANSICHAR* Name, int32 InInternalNumber)
+	: FName(InInternalNumber != NAME_NO_NUMBER_INTERNAL ? FNameHelper::MakeWithInternalNumber(MakeUnconvertedView(Name, Len), FNAME_Add, InInternalNumber)
 		: FNameHelper::MakeDetectNumber(MakeUnconvertedView(Name, Len), FNAME_Add))
 {}
 
-FName::FName(int32 Len, const UTF8CHAR* Name, int32 InNumber)
-	: FName(InNumber != NAME_NO_NUMBER_INTERNAL ? FNameHelper::MakeWithNumber(MakeUnconvertedView(Name, Len), FNAME_Add, InNumber)
+FName::FName(int32 Len, const UTF8CHAR* Name, int32 InInternalNumber)
+	: FName(InInternalNumber != NAME_NO_NUMBER_INTERNAL ? FNameHelper::MakeWithInternalNumber(MakeUnconvertedView(Name, Len), FNAME_Add, InInternalNumber)
 		: FNameHelper::MakeDetectNumber(MakeUnconvertedView(Name, Len), FNAME_Add))
 {}
 
-FName::FName(const TCHAR* Name, int32 InNumber, bool bSplitName)
-	: FName(InNumber == NAME_NO_NUMBER_INTERNAL && bSplitName
+FName::FName(const TCHAR* Name, int32 InInternalNumber, bool bSplitName)
+	: FName(InInternalNumber == NAME_NO_NUMBER_INTERNAL && bSplitName
 		? FNameHelper::MakeDetectNumber(MakeUnconvertedView(Name), FNAME_Add)
-		: FNameHelper::MakeWithNumber(MakeUnconvertedView(Name), FNAME_Add, InNumber))
+		: FNameHelper::MakeWithInternalNumber(MakeUnconvertedView(Name), FNAME_Add, InInternalNumber))
 {}
 
 FName::FName(const FNameEntrySerialized& LoadedEntry)
@@ -4483,7 +4483,7 @@ FNameEntryId FNameEntryId::FromValidEName(EName Ename)
 }
 
 #if UE_FNAME_OUTLINE_NUMBER
-FName FName::CreateNumberedName(FNameEntryId ComparisonId, FNameEntryId DisplayId, int32 Number)
+FName FName::CreateNumberedName(FNameEntryId ComparisonId, FNameEntryId DisplayId, int32 InInternalNumber)
 {
 	// Users must only pass in base/plain ids, not numbered ones.
 #if WITH_CASE_PRESERVING_NAME
@@ -4491,7 +4491,7 @@ FName FName::CreateNumberedName(FNameEntryId ComparisonId, FNameEntryId DisplayI
 #endif
 	checkName(ResolveEntry(ComparisonId)->IsNumbered() == false);
 
-	return FNameHelper::MakeWithNumber(FNameEntryIds{ ComparisonId, DisplayId }, FNAME_Add, Number);
+	return FNameHelper::MakeWithInternalNumber(FNameEntryIds{ ComparisonId, DisplayId }, FNAME_Add, InInternalNumber);
 }
 #endif // UE_FNAME_OUTLINE_NUMBER
 
@@ -4521,7 +4521,7 @@ FName FLazyName::Resolve() const
 		FNameEntryId ComparisonId = Copy.GetComparisonId();
 		FNameEntryId DisplayId = Copy.GetDisplayId();
 #if UE_FNAME_OUTLINE_NUMBER
-		return FNameHelper::MakeWithNumber(FNameEntryIds{ ComparisonId, DisplayId }, FNAME_Add, Number);
+		return FNameHelper::MakeWithInternalNumber(FNameEntryIds{ ComparisonId, DisplayId }, FNAME_Add, Number);
 #else // UE_FNAME_OUTLINE_NUMBER
 		return FName(ComparisonId, DisplayId, Number);
 #endif // UE_FNAME_OUTLINE_NUMBER
@@ -4550,7 +4550,7 @@ FName FLazyName::Resolve() const
 	Either.PackedName = FLiteralOrName::PackName(ComparisonId, DisplayId);
 
 #if UE_FNAME_OUTLINE_NUMBER
-	return FNameHelper::MakeWithNumber(FNameEntryIds{ ComparisonId, DisplayId}, FNAME_Add, Number);
+	return FNameHelper::MakeWithInternalNumber(FNameEntryIds{ ComparisonId, DisplayId}, FNAME_Add, Number);
 #else // UE_FNAME_OUTLINE_NUMBER
 	return FName(ComparisonId, DisplayId, Number);
 #endif // UE_FNAME_OUTLINE_NUMBER
diff --git a/Engine/Source/Runtime/Core/Public/UObject/NameTypes.h b/Engine/Source/Runtime/Core/Public/UObject/NameTypes.h
--- a/Engine/Source/Runtime/Core/Public/UObject/NameTypes.h
+++ b/Engine/Source/Runtime/Core/Public/UObject/NameTypes.h
@@ -638,17 +638,21 @@ public:
 #endif //UE_FNAME_OUTLINE_NUMBER
 
 #if UE_FNAME_OUTLINE_NUMBER
+	// This should really be named GetInternalNumber().
 	[[nodiscard]] CORE_API int32 GetNumber() const;
+	// This should really be named SetInternalNumber().
 	CORE_API void SetNumber(const int32 NewNumber);
 #else //UE_FNAME_OUTLINE_NUMBER
+	// This should really be named GetInternalNumber().
 	[[nodiscard]] FORCEINLINE int32 GetNumber() const
 	{
 		return Number;
 	}
 
-	FORCEINLINE void SetNumber(const int32 NewNumber)
+	// This should really be named SetInternalNumber().
+	FORCEINLINE void SetNumber(const int32 NewInternalNumber)
 	{
-		Number = NewNumber;
+		Number = NewInternalNumber;
 	}
 #endif //UE_FNAME_OUTLINE_NUMBER
 	
@@ -957,10 +961,10 @@ public:
 	 * Create an FName with a hardcoded string index and (instance).
 	 *
 	 * @param N The hardcoded value the string portion of the name will have
-	 * @param InNumber The hardcoded value for the number portion of the name
+	 * @param InInternalNumber The hardcoded value for the number portion of the name, in internal format
 	 */
-	FORCEINLINE FName(EName Ename, int32 InNumber)
-		: FName(CreateNumberedNameIfNecessary(FNameEntryId::FromEName(Ename), InNumber))
+	FORCEINLINE FName(EName Ename, int32 InInternalNumber)
+		: FName(CreateNumberedNameIfNecessary(FNameEntryId::FromEName(Ename), InInternalNumber))
 	{
 	}
 
@@ -969,10 +973,10 @@ public:
 	 * Create an FName from an existing string, but with a different instance.
 	 *
 	 * @param Other The FName to take the string values from
-	 * @param InNumber The hardcoded value for the number portion of the name
+	 * @param InInternalNumber The hardcoded value for the number portion of the name, in internal format
 	 */
-	FORCEINLINE FName(FName Other, int32 InNumber)
-		: FName(CreateNumberedNameIfNecessary(Other.GetComparisonIndex(), Other.GetDisplayIndex(), InNumber))
+	FORCEINLINE FName(FName Other, int32 InInternalNumber)
+		: FName(CreateNumberedNameIfNecessary(Other.GetComparisonIndex(), Other.GetDisplayIndex(), InInternalNumber))
 	{
 	}
 
@@ -980,8 +984,8 @@ public:
 	 * Create an FName from its component parts
 	 * Only call this if you *really* know what you're doing
 	 */
-	FORCEINLINE FName(FNameEntryId InComparisonIndex, FNameEntryId InDisplayIndex, int32 InNumber)
-		: FName(CreateNumberedNameIfNecessary(InComparisonIndex, InDisplayIndex, InNumber))
+	FORCEINLINE FName(FNameEntryId InComparisonIndex, FNameEntryId InDisplayIndex, int32 InInternalNumber)
+		: FName(CreateNumberedNameIfNecessary(InComparisonIndex, InDisplayIndex, InInternalNumber))
 	{
 	}
 
@@ -994,16 +998,16 @@ public:
 	/**
 	 * Only call this if you *really* know what you're doing
 	 */
-	[[nodiscard]] static FName CreateFromDisplayId(FNameEntryId DisplayId, int32 Number)
+	[[nodiscard]] static FName CreateFromDisplayId(FNameEntryId DisplayId, int32 InInternalNumber)
 	{
 #if UE_FNAME_OUTLINE_NUMBER
 		checkSlow(ResolveEntry(DisplayId)->IsNumbered() == false); // This id should be unnumbered i.e. returned from GetDisplayIndex on an FName.
 #endif
-		return FName(GetComparisonIdFromDisplayId(DisplayId), DisplayId, Number);
+		return FName(GetComparisonIdFromDisplayId(DisplayId), DisplayId, InInternalNumber);
 	}
 
 #if UE_FNAME_OUTLINE_NUMBER
-	CORE_API static FName FindNumberedName(FNameEntryId DisplayId, int32 Number);
+	CORE_API static FName FindNumberedName(FNameEntryId DisplayId, int32 InInternalNumber);
 #endif //UE_FNAME_OUTLINE_NUMBER
 
 	/**
@@ -1087,41 +1091,39 @@ public:
 	 * When UE_FNAME_OUTLINE_NUMBER is set, will also add the combination of base string and number to the name table if it doesn't exist.
 	 *
 	 * @param Name Value for the string portion of the name
-	 * @param Number Value for the number portion of the name
+	 * @param InInternalNumber Value for the number portion of the name, in internal format
 	 */
-	CORE_API FName(const WIDECHAR* Name, int32 Number);
-	CORE_API FName(const ANSICHAR* Name, int32 Number);
-	CORE_API FName(const UTF8CHAR* Name, int32 Number);
-	CORE_API FName(int32 Len, const WIDECHAR* Name, int32 Number);
-	CORE_API FName(int32 Len, const ANSICHAR* Name, int32 Number);
-	CORE_API FName(int32 Len, const UTF8CHAR* Name, int32 Number);
-
-	inline FName(TStringView<ANSICHAR> View, int32 InNumber)
+	CORE_API FName(const WIDECHAR* Name, int32 InInternalNumber);
+	CORE_API FName(const ANSICHAR* Name, int32 InInternalNumber);
+	CORE_API FName(const UTF8CHAR* Name, int32 InInternalNumber);
+	CORE_API FName(int32 Len, const WIDECHAR* Name, int32 InInternalNumber);
+	CORE_API FName(int32 Len, const ANSICHAR* Name, int32 InInternalNumber);
+	CORE_API FName(int32 Len, const UTF8CHAR* Name, int32 InInternalNumber);
+
+	inline FName(TStringView<ANSICHAR> View, int32 InInternalNumber)
 		: FName(NoInit)
 	{
-		*this = FName(View.Len(), View.GetData(), InNumber);
+		*this = FName(View.Len(), View.GetData(), InInternalNumber);
 	}
-	inline FName(TStringView<WIDECHAR> View, int32 InNumber)
+	inline FName(TStringView<WIDECHAR> View, int32 InInternalNumber)
 		: FName(NoInit)
 	{
-		*this = FName(View.Len(), View.GetData(), InNumber);
+		*this = FName(View.Len(), View.GetData(), InInternalNumber);
 	}
-	inline FName(TStringView<UTF8CHAR> View, int32 InNumber)
+	inline FName(TStringView<UTF8CHAR> View, int32 InInternalNumber)
 		: FName(NoInit)
 	{
-		*this = FName(View.Len(), View.GetData(), InNumber);
+		*this = FName(View.Len(), View.GetData(), InInternalNumber);
 	}
 
 	/**
-	 * Create an FName. If FindType is FNAME_Find, and the string part of the name 
-	 * doesn't already exist, then the name will be NAME_None
+	 * Create an FName.
 	 *
 	 * @param Name Value for the string portion of the name
-	 * @param Number Value for the number portion of the name
-	 * @param FindType Action to take (see EFindName)
+	 * @param InInternalNumber Value for the number portion of the name, in internal format
 	 * @param bSplitName true if the trailing number should be split from the name when Number == NAME_NO_NUMBER_INTERNAL, or false to always use the name as-is
 	 */
-	CORE_API FName(const TCHAR* Name, int32 InNumber, bool bSplitName);
+	CORE_API FName(const TCHAR* Name, int32 InInternalNumber, bool bSplitName);
 
 	/**
 	 * Constructor used by FLinkerLoad when loading its name table; Creates an FName with an instance
@@ -1292,16 +1294,16 @@ private:
 
 	// These FNameEntryIds are passed in from user code so they must be non-numbered if Number != NAME_NO_NUMBER_INTERNAL
 #if UE_FNAME_OUTLINE_NUMBER
-	[[nodiscard]] CORE_API static FName CreateNumberedName(FNameEntryId ComparisonId, FNameEntryId DisplayId, int32 Number);
+	[[nodiscard]] CORE_API static FName CreateNumberedName(FNameEntryId ComparisonId, FNameEntryId DisplayId, int32 InInternalNumber);
 #endif
 
-	[[nodiscard]] FORCEINLINE static FName CreateNumberedNameIfNecessary(FNameEntryId ComparisonId, FNameEntryId DisplayId, int32 Number)
+	[[nodiscard]] FORCEINLINE static FName CreateNumberedNameIfNecessary(FNameEntryId ComparisonId, FNameEntryId DisplayId, int32 InInternalNumber)
 	{
 #if UE_FNAME_OUTLINE_NUMBER
-		if (Number != NAME_NO_NUMBER_INTERNAL)
+		if (InInternalNumber != NAME_NO_NUMBER_INTERNAL)
 		{
 			// We need to store a new entry in the name table
-			return CreateNumberedName(ComparisonId, DisplayId, Number);
+			return CreateNumberedName(ComparisonId, DisplayId, InInternalNumber);
 		}
 		// Otherwise we can just set the index members
 #endif
@@ -1311,14 +1313,14 @@ private:
 		Out.DisplayIndex = DisplayId;
 #endif
 #if !UE_FNAME_OUTLINE_NUMBER
-		Out.Number = Number;
+		Out.Number = InInternalNumber;
 #endif
 		return Out;
 	}
 
-	[[nodiscard]] FORCEINLINE static FName CreateNumberedNameIfNecessary(FNameEntryId ComparisonId, int32 Number)
+	[[nodiscard]] UE_REWRITE static FName CreateNumberedNameIfNecessary(FNameEntryId ComparisonId, int32 InInternalNumber)
 	{
-		return CreateNumberedNameIfNecessary(ComparisonId, ComparisonId, Number);
+		return CreateNumberedNameIfNecessary(ComparisonId, ComparisonId, InInternalNumber);
 	}
 	
 	
@@ -1858,7 +1860,7 @@ class FDisplayNameEntryId
 public:
 	FDisplayNameEntryId() : FDisplayNameEntryId(FName()) {}
 	explicit FDisplayNameEntryId(FName Name) : FDisplayNameEntryId(Name.GetDisplayIndex(), Name.GetComparisonIndex()) {}
-	FORCEINLINE FName ToName(uint32 Number) const { return FName(GetComparisonId(), GetDisplayId(), Number); }
+	FORCEINLINE FName ToName(uint32 InInternalNumber) const { return FName(GetComparisonId(), GetDisplayId(), InInternalNumber); }
 
 private:
 #if WITH_CASE_PRESERVING_NAME
