Add validation that pthread ID is valid

#Synced-CL 51050668

diff --git a/Engine/Source/Runtime/Core/Private/HAL/PThreadRunnableThread.h b/Engine/Source/Runtime/Core/Private/HAL/PThreadRunnableThread.h
--- a/Engine/Source/Runtime/Core/Private/HAL/PThreadRunnableThread.h
+++ b/Engine/Source/Runtime/Core/Private/HAL/PThreadRunnableThread.h
@@ -177,6 +177,11 @@ protected:
 		// cache the thread ID for this thread (defined by the platform)
 		ThisThread->ThreadID = FPlatformTLS::GetCurrentThreadId();
 
+		if (ThisThread->ThreadID == ~uint32(0)) [[unlikely]]
+		{
+			LowLevelFatalError(TEXT("Thread ID is -1. Reconsider implementation of FThread::InvalidThreadId"));
+		}
+
 		FThreadManager::Get().AddThread(ThisThread->ThreadID, ThisThread);
 
 		// set the affinity.  This function sets affinity on the current thread, so don't call in the Create function which will trash the main thread affinity.
