Logging: Fixed support for UE_LOGF with Fatal verbosity

#rnx
#Synced-CL 50967057

diff --git a/Engine/Source/Runtime/Core/Private/Logging/StructuredLog.cpp b/Engine/Source/Runtime/Core/Private/Logging/StructuredLog.cpp
--- a/Engine/Source/Runtime/Core/Private/Logging/StructuredLog.cpp
+++ b/Engine/Source/Runtime/Core/Private/Logging/StructuredLog.cpp
@@ -37,6 +37,7 @@
 
 void StaticFailDebug (const TCHAR* Error, const ANSICHAR* File, int32 Line, void* ProgramCounter, const TCHAR* Message);
 void StaticFailDebugV(const TCHAR* Error, const ANSICHAR* File, int32 Line, void* ProgramCounter, const TCHAR* DescriptionFormat, va_list DescriptionArgs);
+void StaticFailDebugV(const TCHAR* Error, const ANSICHAR* File, int32 Line, void* ProgramCounter, const ANSICHAR* DescriptionFormat, va_list DescriptionArgs);
 void StaticFailDebugV(const TCHAR* Error, const ANSICHAR* File, int32 Line, void* ProgramCounter, const UTF8CHAR* DescriptionFormat, va_list DescriptionArgs);
 
 namespace UE::Serialization::Private
@@ -1789,6 +1790,7 @@ FORCENOINLINE void BasicFatalLog(TStaticBasicLogRecordParam<FormatCharType> Log,
 }
 
 template CORE_API void BasicFatalLog(TStaticBasicLogRecordParam<TCHAR> Log, const FLogCategoryBase* Category, ...);
+template CORE_API void BasicFatalLog(TStaticBasicLogRecordParam<ANSICHAR> Log, const FLogCategoryBase* Category, ...);
 template CORE_API void BasicFatalLog(TStaticBasicLogRecordParam<UTF8CHAR> Log, const FLogCategoryBase* Category, ...);
 
 template <UE::CCharType FormatCharType>
@@ -1803,6 +1805,7 @@ FORCENOINLINE void BasicFatalLogWithProgramCounter(TStaticBasicLogRecordParam<Fo
 }
 
 template CORE_API void BasicFatalLogWithProgramCounter(TStaticBasicLogRecordParam<TCHAR> Log, const FLogCategoryBase* Category, void* ProgramCounter, ...);
+template CORE_API void BasicFatalLogWithProgramCounter(TStaticBasicLogRecordParam<ANSICHAR> Log, const FLogCategoryBase* Category, void* ProgramCounter, ...);
 template CORE_API void BasicFatalLogWithProgramCounter(TStaticBasicLogRecordParam<UTF8CHAR> Log, const FLogCategoryBase* Category, void* ProgramCounter, ...);
 
 #if !NO_LOGGING
diff --git a/Engine/Source/Runtime/Core/Private/Misc/AssertionMacros.cpp b/Engine/Source/Runtime/Core/Private/Misc/AssertionMacros.cpp
--- a/Engine/Source/Runtime/Core/Private/Misc/AssertionMacros.cpp
+++ b/Engine/Source/Runtime/Core/Private/Misc/AssertionMacros.cpp
@@ -257,6 +257,11 @@ public:
 		}
 	}
 
+	UE_REWRITE void AppendV(const ANSICHAR* Format, va_list Args)
+	{
+		AppendV((const UTF8CHAR*)Format, Args);
+	}
+
 	void Appendf(const TCHAR* Format, ...)
 	{
 		va_list Args;
@@ -347,6 +352,11 @@ FORCENOINLINE void VARARGS StaticFailDebugV(const TCHAR* Error, const ANSICHAR*
 	StaticFailDebugV(Error, /*Expression*/ "", File, Line, /*bIsEnsure*/ false, ProgramCounter, /*Message*/ nullptr, DescriptionFormat, DescriptionArgs);
 }
 
+FORCENOINLINE void VARARGS StaticFailDebugV(const TCHAR* Error, const ANSICHAR* File, int32 Line, void* ProgramCounter, const ANSICHAR* DescriptionFormat, va_list DescriptionArgs)
+{
+	StaticFailDebugV(Error, /*Expression*/ "", File, Line, /*bIsEnsure*/ false, ProgramCounter, /*Message*/ nullptr, DescriptionFormat, DescriptionArgs);
+}
+
 FORCENOINLINE void VARARGS StaticFailDebugV(const TCHAR* Error, const ANSICHAR* File, int32 Line, void* ProgramCounter, const UTF8CHAR* DescriptionFormat, va_list DescriptionArgs)
 {
 	StaticFailDebugV(Error, /*Expression*/ "", File, Line, /*bIsEnsure*/ false, ProgramCounter, /*Message*/ nullptr, DescriptionFormat, DescriptionArgs);
