Deprecated PLATFORM_ENABLE_POPCNT_INTRINSIC

Also deprecated the associated FPlatformMisc::HasNonoptionalCPUFeatures() and FPlatformMisc::NeedsNonoptionalCPUFeaturesCheck().

#Synced-CL 50967057

diff --git a/Engine/Source/Runtime/Core/Private/GenericPlatform/GenericPlatformMisc.cpp b/Engine/Source/Runtime/Core/Private/GenericPlatform/GenericPlatformMisc.cpp
--- a/Engine/Source/Runtime/Core/Private/GenericPlatform/GenericPlatformMisc.cpp
+++ b/Engine/Source/Runtime/Core/Private/GenericPlatform/GenericPlatformMisc.cpp
@@ -859,18 +859,6 @@ uint32 FGenericPlatformMisc::GetCPUInfo()
 	return 0;
 }
 
-bool FGenericPlatformMisc::HasNonoptionalCPUFeatures()
-{
-	// Not implemented cross-platform. Each platform may or may not choose to implement this.
-	return false;
-}
-
-bool FGenericPlatformMisc::NeedsNonoptionalCPUFeaturesCheck()
-{
-	// This is opt in on a per platform basis
-	return false;
-}
-
 FString FGenericPlatformMisc::GetHostArchitecture()
 {
 	return FString();
diff --git a/Engine/Source/Runtime/Core/Private/Unix/UnixPlatformMisc.cpp b/Engine/Source/Runtime/Core/Private/Unix/UnixPlatformMisc.cpp
--- a/Engine/Source/Runtime/Core/Private/Unix/UnixPlatformMisc.cpp
+++ b/Engine/Source/Runtime/Core/Private/Unix/UnixPlatformMisc.cpp
@@ -1422,45 +1422,6 @@ FString FUnixPlatformMisc::GetPrimaryGPUBrand()
 	return ProcGPUBrandStr;
 }
 
-// __builtin_popcountll() will not be compiled to use popcnt instruction unless -mpopcnt or a sufficiently recent target CPU arch is passed (which UBT doesn't by default)
-#if defined(__POPCNT__)
-	#define UE4_LINUX_NEED_TO_CHECK_FOR_POPCNT_PRESENCE				(PLATFORM_ENABLE_POPCNT_INTRINSIC)
-#else
-	#define UE4_LINUX_NEED_TO_CHECK_FOR_POPCNT_PRESENCE				0
-#endif // __POPCNT__
-
-bool FUnixPlatformMisc::HasNonoptionalCPUFeatures()
-{
-	static bool bHasNonOptionalFeature = false;
-	static bool bHaveResult = false;
-
-	if (!bHaveResult)
-	{
-#if PLATFORM_HAS_CPUID
-		int Info[4];
-		__cpuid(1, Info[0], Info[1], Info[2], Info[3]);
-
-	#if UE4_LINUX_NEED_TO_CHECK_FOR_POPCNT_PRESENCE
-		bHasNonOptionalFeature = (Info[2] & (1 << 23)) != 0;
-	#endif // UE4_LINUX_NEED_TO_CHECK_FOR_POPCNT
-#endif // PLATFORM_HAS_CPUID
-
-		bHaveResult = true;
-	}
-
-	return bHasNonOptionalFeature;
-}
-
-bool FUnixPlatformMisc::NeedsNonoptionalCPUFeaturesCheck()
-{
-#if UE4_LINUX_NEED_TO_CHECK_FOR_POPCNT_PRESENCE
-	return true;
-#else
-	return false;
-#endif
-}
-
-
 #if !UE_BUILD_SHIPPING
 bool FUnixPlatformMisc::IsDebuggerPresent()
 {
diff --git a/Engine/Source/Runtime/Core/Private/Windows/WindowsPlatformMisc.cpp b/Engine/Source/Runtime/Core/Private/Windows/WindowsPlatformMisc.cpp
--- a/Engine/Source/Runtime/Core/Private/Windows/WindowsPlatformMisc.cpp
+++ b/Engine/Source/Runtime/Core/Private/Windows/WindowsPlatformMisc.cpp
@@ -4103,18 +4103,6 @@ uint32 FWindowsPlatformMisc::GetCPUInfo()
 	return FCPUIDQueriedData::GetCPUInfo();
 }
 
-bool FWindowsPlatformMisc::HasNonoptionalCPUFeatures()
-{
-	// Check for popcnt is bit 23
-	return (FCPUIDQueriedData::GetCPUInfo2() & (1 << 23)) != 0;
-}
-
-bool FWindowsPlatformMisc::NeedsNonoptionalCPUFeaturesCheck()
-{
-	// popcnt is 64bit
-	return PLATFORM_ENABLE_POPCNT_INTRINSIC;
-}
-
 bool FWindowsPlatformMisc::HasTimedPauseCPUFeature()
 {
 	return FCPUIDQueriedData::HasTimedPauseInstruction();
diff --git a/Engine/Source/Runtime/Core/Public/GenericPlatform/GenericPlatformMisc.h b/Engine/Source/Runtime/Core/Public/GenericPlatform/GenericPlatformMisc.h
--- a/Engine/Source/Runtime/Core/Public/GenericPlatform/GenericPlatformMisc.h
+++ b/Engine/Source/Runtime/Core/Public/GenericPlatform/GenericPlatformMisc.h
@@ -804,9 +804,11 @@ struct FGenericPlatformMisc
 	static CORE_API uint32 GetCPUInfo();
 
 	/** @return whether this cpu supports certain required instructions or not */
-	static CORE_API bool HasNonoptionalCPUFeatures();
+	UE_DEPRECATED(5.8, "Replace HasNonoptionalCPUFeatures() with false.")
+	static inline bool HasNonoptionalCPUFeatures() { return false; }
 	/** @return whether to check for specific CPU compatibility or not */
-	static CORE_API bool NeedsNonoptionalCPUFeaturesCheck();
+	UE_DEPRECATED(5.8, "Replace NeedsNonoptionalCPUFeaturesCheck() with false.")
+	static inline bool NeedsNonoptionalCPUFeaturesCheck() { return false; }
 
 	/** @return architecture of the host machine (x64, arm64 etc) **/
 	static CORE_API FString GetHostArchitecture();
diff --git a/Engine/Source/Runtime/Core/Public/HAL/Platform.h b/Engine/Source/Runtime/Core/Public/HAL/Platform.h
--- a/Engine/Source/Runtime/Core/Public/HAL/Platform.h
+++ b/Engine/Source/Runtime/Core/Public/HAL/Platform.h
@@ -243,11 +243,7 @@
 		#define PLATFORM_HAS_CPUID				0
 	#endif
 #endif
-#ifndef PLATFORM_ENABLE_POPCNT_INTRINSIC
-	// PC is disabled by default, but linux and mac are enabled
-	// if your min spec is an AMD cpu mid-2007 or Intel 2008, you should enable this
-	#define PLATFORM_ENABLE_POPCNT_INTRINSIC 0
-#endif
+#define PLATFORM_ENABLE_POPCNT_INTRINSIC UE_DEPRECATED_MACRO(5.8, "PLATFORM_ENABLE_POPCNT_INTRINSIC has been deprecated and should be replaced with 1.") 1
 #ifndef PLATFORM_ENABLE_VECTORINTRINSICS_NEON
 	#define PLATFORM_ENABLE_VECTORINTRINSICS_NEON	0
 #endif
diff --git a/Engine/Source/Runtime/Core/Public/Mac/MacPlatform.h b/Engine/Source/Runtime/Core/Public/Mac/MacPlatform.h
--- a/Engine/Source/Runtime/Core/Public/Mac/MacPlatform.h
+++ b/Engine/Source/Runtime/Core/Public/Mac/MacPlatform.h
@@ -52,8 +52,6 @@ typedef FMacPlatformTypes FPlatformTypes;
 #define PLATFORM_SUPPORTS_MESH_SHADERS					1
 #define PLATFORM_SUPPORTS_BINDLESS_RENDERING			1
 #define PLATFORM_SUPPORTS_GEOMETRY_SHADERS				1
- 
-#define PLATFORM_ENABLE_POPCNT_INTRINSIC				1
 
 #define PLATFORM_GLOBAL_LOG_CATEGORY					LogMac
 
diff --git a/Engine/Source/Runtime/Core/Public/Mac/MacPlatformMisc.h b/Engine/Source/Runtime/Core/Public/Mac/MacPlatformMisc.h
--- a/Engine/Source/Runtime/Core/Public/Mac/MacPlatformMisc.h
+++ b/Engine/Source/Runtime/Core/Public/Mac/MacPlatformMisc.h
@@ -92,9 +92,6 @@ struct CORE_API FMacPlatformMisc : public FApplePlatformMisc
 	 */
 	static uint32 GetCPUInfo();
 
-    static bool HasNonoptionalCPUFeatures() { return true; }
-    static bool NeedsNonoptionalCPUFeaturesCheck() { return PLATFORM_ENABLE_POPCNT_INTRINSIC; }
-
 	static void SetGracefulTerminationHandler();
 	
 	static void SetCrashHandler(void (* CrashHandler)(const FGenericCrashContext& Context));
diff --git a/Engine/Source/Runtime/Core/Public/Unix/UnixPlatform.h b/Engine/Source/Runtime/Core/Public/Unix/UnixPlatform.h
--- a/Engine/Source/Runtime/Core/Public/Unix/UnixPlatform.h
+++ b/Engine/Source/Runtime/Core/Public/Unix/UnixPlatform.h
@@ -78,8 +78,6 @@ typedef FUnixPlatformTypes FPlatformTypes;
 
 #define PLATFORM_CODE_SECTION(Name)						__attribute__((section(Name)))
 
-#define PLATFORM_ENABLE_POPCNT_INTRINSIC				1
-
 // SOCK_CLOEXEC is available on Unix since 2.6.27
 #if LINUX_VERSION_CODE >= KERNEL_VERSION(2,6,27)
 	#define PLATFORM_HAS_BSD_SOCKET_FEATURE_CLOSE_ON_EXEC	1
diff --git a/Engine/Source/Runtime/Core/Public/Unix/UnixPlatformMisc.h b/Engine/Source/Runtime/Core/Public/Unix/UnixPlatformMisc.h
--- a/Engine/Source/Runtime/Core/Public/Unix/UnixPlatformMisc.h
+++ b/Engine/Source/Runtime/Core/Public/Unix/UnixPlatformMisc.h
@@ -101,9 +101,6 @@ struct FUnixPlatformMisc : public FGenericPlatformMisc
 
 	static CORE_API FString GetPrimaryGPUBrand();
 
-	static CORE_API bool HasNonoptionalCPUFeatures();
-	static CORE_API bool NeedsNonoptionalCPUFeaturesCheck();
-
 	/**
 	 * Ungrabs input (useful before breaking into debugging)
 	 */
diff --git a/Engine/Source/Runtime/Core/Public/Windows/WindowsPlatformMisc.h b/Engine/Source/Runtime/Core/Public/Windows/WindowsPlatformMisc.h
--- a/Engine/Source/Runtime/Core/Public/Windows/WindowsPlatformMisc.h
+++ b/Engine/Source/Runtime/Core/Public/Windows/WindowsPlatformMisc.h
@@ -283,10 +283,6 @@ struct FWindowsPlatformMisc
 	 */
 	static CORE_API uint32 GetCPUInfo();
 
-	/** @return whether this cpu supports certain required instructions or not */
-	static CORE_API bool HasNonoptionalCPUFeatures();
-	/** @return whether to check for specific CPU compatibility or not */
-	static CORE_API bool NeedsNonoptionalCPUFeaturesCheck();
 	/** @return whether this cpu has timed pause instruction support or not */
 	static CORE_API bool HasTimedPauseCPUFeature();
 
