[IOS] Do not fatal on bad entitlements in shipping

This triggers the crash reporting system, but it is not actually a legit crash.

When users are messing with the IPA (e.g. resigning with wron entitlements), it is fine to just force exit.

#rnx
#Synced-CL 51020593

diff --git a/Engine/Source/Runtime/Core/Private/IOS/IOSPlatformMisc.cpp b/Engine/Source/Runtime/Core/Private/IOS/IOSPlatformMisc.cpp
--- a/Engine/Source/Runtime/Core/Private/IOS/IOSPlatformMisc.cpp
+++ b/Engine/Source/Runtime/Core/Private/IOS/IOSPlatformMisc.cpp
@@ -221,7 +221,13 @@ void FIOSPlatformMisc::PlatformInit()
 	{
 		if (!FIOSPlatformMisc::IsEntitlementEnabled(TCHAR_TO_ANSI(*Entitlement)))
 		{
+#if UE_BUILD_SHIPPING
+			// We don't want this to be counted as a legit crash, as this should only
+			// happen if users are messing with the IPA.
+			_Exit(0);
+#else
 			UE_LOG(LogInit, Fatal, TEXT("App does not have required entitlement %s."), *Entitlement);
+#endif
 		}
 	}
 
