Added UE_FORCEINLINE_HINT to Align templates to make sure they get inlined instead of a call.

#Synced-CL 51061797

diff --git a/Engine/Source/Runtime/Core/Public/Templates/AlignmentTemplates.h b/Engine/Source/Runtime/Core/Public/Templates/AlignmentTemplates.h
--- a/Engine/Source/Runtime/Core/Public/Templates/AlignmentTemplates.h
+++ b/Engine/Source/Runtime/Core/Public/Templates/AlignmentTemplates.h
@@ -16,7 +16,7 @@
  */
 template <typename T>
 	requires UE::CIntegral<T> || UE::CPointer<T>
-inline constexpr T Align(T Val, uint64 Alignment)
+UE_FORCEINLINE_HINT constexpr T Align(T Val, uint64 Alignment)
 {
 	return (T)(((uint64)Val + Alignment - 1) & ~(Alignment - 1));
 }
@@ -31,7 +31,7 @@ inline constexpr T Align(T Val, uint64 Alignment)
  */
 template <typename T>
 	requires UE::CIntegral<T> || UE::CPointer<T>
-inline constexpr T AlignDown(T Val, uint64 Alignment)
+UE_FORCEINLINE_HINT constexpr T AlignDown(T Val, uint64 Alignment)
 {
 	return (T)(((uint64)Val) & ~(Alignment - 1));
 }
@@ -46,7 +46,7 @@ inline constexpr T AlignDown(T Val, uint64 Alignment)
  */
 template <typename T>
 	requires UE::CIntegral<T> || UE::CPointer<T>
-inline constexpr bool IsAligned(T Val, uint64 Alignment)
+UE_FORCEINLINE_HINT constexpr bool IsAligned(T Val, uint64 Alignment)
 {
 	return !((uint64)Val & (Alignment - 1));
 }
@@ -61,7 +61,7 @@ inline constexpr bool IsAligned(T Val, uint64 Alignment)
  */
 template <typename T>
 	requires UE::CIntegral<T> || UE::CPointer<T>
-inline constexpr T AlignArbitrary(T Val, uint64 Alignment)
+UE_FORCEINLINE_HINT constexpr T AlignArbitrary(T Val, uint64 Alignment)
 {
 	return (T)((((uint64)Val + Alignment - 1) / Alignment) * Alignment);
 }
