Improve logging around SafeReload

#rnx

#Synced-CL 51088011

diff --git a/Engine/Source/Runtime/Core/Private/Misc/ConfigCacheIni.cpp b/Engine/Source/Runtime/Core/Private/Misc/ConfigCacheIni.cpp
--- a/Engine/Source/Runtime/Core/Private/Misc/ConfigCacheIni.cpp
+++ b/Engine/Source/Runtime/Core/Private/Misc/ConfigCacheIni.cpp
@@ -15,6 +15,8 @@
 #include "HAL/IConsoleManager.h"
 #include "HAL/LowLevelMemStats.h"
 #include "HAL/LowLevelMemTracker.h"
+#include "HAL/PlatformTLS.h"
+#include "HAL/ThreadManager.h"
 #include "Misc/CoreDelegates.h"
 #include "Misc/App.h"
 #include "Misc/RemoteConfigIni.h"
@@ -113,6 +115,68 @@ static int GDefaultReplayMethod = 1;
 
 static bool OverrideFileFromCommandline(FString & InOutFilename);
 
+/**
+ * Stores some values that are helpful for logging but would be cumbersome to pass around to every internal function.
+ */
+struct FConfigLogContext
+{
+	FConfigLogContext* OuterContext = nullptr;
+	int32 ContextDepth = 0;
+
+	const TCHAR* Action = nullptr;
+	const TCHAR* Section = nullptr;
+	const TCHAR* Key = nullptr;
+};
+
+thread_local static FConfigLogContext ConfigLogContext;
+
+/**
+ * Pushes a new set of logging context values for the current scope.
+ */
+struct FConfigLogContextScope
+{
+	FConfigLogContext PrevLogContext;
+
+	FConfigLogContextScope(const TCHAR* InAction, const TCHAR* InSection, const TCHAR* InKey)
+	{
+		PrevLogContext = ConfigLogContext;
+
+		// Only set the outer pointer if this isn't the first scope
+		if (ConfigLogContext.ContextDepth > 0)
+		{
+			ConfigLogContext.OuterContext = &PrevLogContext;
+		}
+
+		ConfigLogContext.ContextDepth++;
+		ConfigLogContext.Action = InAction;
+		ConfigLogContext.Section = InSection;
+		ConfigLogContext.Key = InKey;
+	}
+
+	~FConfigLogContextScope()
+	{
+		// Verify OuterContext pointer was not modified
+		check((ConfigLogContext.OuterContext == nullptr) || (ConfigLogContext.OuterContext == &PrevLogContext));
+
+		ConfigLogContext = PrevLogContext;
+	}
+};
+
+/**
+ * Gets the "outermost" log context. Useful for determining which function was called from outside code.
+ * e.g. FindConfigFile can be called directly, but is also called by Find and FindPlatformConfig.
+ */
+FConfigLogContext* GetOutermostConfigLogContext()
+{
+	FConfigLogContext* Result = &ConfigLogContext;
+	while (Result->OuterContext)
+	{
+		Result = Result->OuterContext;
+	}
+
+	return Result;
+}
+
 /*-----------------------------------------------------------------------------
 FConfigValue
 -----------------------------------------------------------------------------*/
@@ -3916,6 +3980,12 @@ void FConfigBranch::SafeReload()
 	const float ReloadTimeMs = (float)(FPlatformTime::Seconds() - StartTime) * 1000.0f;
 	UE_LOG(LogConfig, Log, TEXT("Branch '%s' had been unloaded. Reloading on-demand took %.2fms"), *IniNameStr, ReloadTimeMs);
 
+	const FConfigLogContext* LogContext = GetOutermostConfigLogContext();
+	const TCHAR* ContextAction = LogContext->Action ? LogContext->Action : TEXT("");
+	const TCHAR* ContextSection = LogContext->Section ? LogContext->Section : TEXT("");
+	const TCHAR* ContextKey = LogContext->Key ? LogContext->Key : TEXT("");
+	UE_LOGF(LogConfig, Log, "Reload Context for thread %ls: Action=%ls Section=%ls Key=%ls", *FThreadManager::GetThreadName(FPlatformTLS::GetCurrentThreadId()), ContextAction, ContextSection, ContextKey);
+
 	FCoreDelegates::TSOnConfigBranchReloaded().Broadcast(*IniNameStr, ReloadTimeMs);
 }
 
@@ -4241,6 +4311,8 @@ FConfigBranch* FConfigCacheIni::FindBranchWithNoReload(FName BaseIniName, const
 
 FConfigBranch* FConfigCacheIni::FindBranch(FName BaseIniName, const FString& Filename)
 {
+	FConfigLogContextScope LogScope(TEXT("FindBranch"), nullptr, nullptr);
+
 	FConfigBranch* Branch = FindBranchWithNoReload(BaseIniName, Filename);
 
 	if (Branch)
@@ -4294,6 +4366,8 @@ int32 FConfigCacheIni::Remove(const FString& Filename)
 
 FConfigFile* FConfigCacheIni::FindConfigFile( const FString& Filename )
 {
+	FConfigLogContextScope LogScope(TEXT("FindConfigFile"), nullptr, nullptr);
+
 	FConfigBranch* Result;
 	if (!Filename.EndsWith(TEXT(".ini")))
 	{
@@ -4323,6 +4397,8 @@ FConfigFile* FConfigCacheIni::FindConfigFile( const FString& Filename )
 
 FConfigFile* FConfigCacheIni::Find(const FString& Filename)
 {
+	FConfigLogContextScope LogScope(TEXT("Find"), nullptr, nullptr);
+
 	// check for non-filenames
 	if(Filename.Len() == 0)
 	{
@@ -4375,6 +4451,8 @@ FConfigFile* FConfigCacheIni::Find(const FString& Filename)
 
 FConfigFile* FConfigCacheIni::FindConfigFileWithBaseName(FName BaseName)
 {
+	FConfigLogContextScope LogScope(TEXT("FindConfigFileWithBaseName"), nullptr, nullptr);
+
 	FConfigBranch* Result = KnownFiles.GetBranch(BaseName);
 	if (Result == nullptr)
 	{
@@ -4582,6 +4660,8 @@ bool FConfigCacheIni::AreFileOperationsDisabled()
  */
 void FConfigCacheIni::Parse1ToNSectionOfNames(const TCHAR* Section, const TCHAR* KeyOne, const TCHAR* KeyN, TMap<FName, TArray<FName> >& OutMap, const FString& Filename)
 {
+	FConfigLogContextScope LogScope(TEXT("Parse1ToNSectionOfNames"), Section, KeyOne);
+
 	// find the config file object
 	FConfigFile* ConfigFile = Find(Filename);
 	if (!ConfigFile)
@@ -4647,6 +4727,8 @@ void FConfigCacheIni::Parse1ToNSectionOfNames(const TCHAR* Section, const TCHAR*
  */
 void FConfigCacheIni::Parse1ToNSectionOfStrings(const TCHAR* Section, const TCHAR* KeyOne, const TCHAR* KeyN, TMap<FString, TArray<FString> >& OutMap, const FString& Filename)
 {
+	FConfigLogContextScope LogScope(TEXT("Parse1ToNSectionOfStrings"), Section, KeyOne);
+
 	// find the config file object
 	FConfigFile* ConfigFile = Find(Filename);
 	if (!ConfigFile)
@@ -4765,6 +4847,8 @@ static FConfigBranch* FindOrLoadNoSafeReload(FConfigCacheIni* ConfigSystem, cons
 
 bool FConfigCacheIni::GetString( const TCHAR* Section, const TCHAR* Key, FString& Value, const FString& Filename )
 {
+	FConfigLogContextScope LogScope(TEXT("GetString"), Section, Key);
+
 	// get the file without a SafeReload possible, so no lock is needed
 	FConfigBranch* Branch = FindOrLoadNoSafeReload(this, Filename);
 	if (Branch == nullptr)
@@ -4783,7 +4867,6 @@ bool FConfigCacheIni::GetString( const TCHAR* Section, const TCHAR* Key, FString
 	
 	if (Branch->bIsSafeUnloaded)
 	{
-		UE_LOG(LogConfig, Log, TEXT("About to reload due to getting %s: %s"), Section, Key);
 		Branch->SafeReload();
 	}
 #else
@@ -4817,6 +4900,8 @@ bool FConfigCacheIni::GetString( const TCHAR* Section, const TCHAR* Key, FString
 
 bool FConfigCacheIni::GetText( const TCHAR* Section, const TCHAR* Key, FText& Value, const FString& Filename )
 {
+	FConfigLogContextScope LogScope(TEXT("GetText"), Section, Key);
+
 	// get the file without a SafeReload possible, so no lock is needed
 	FConfigBranch* Branch = FindOrLoadNoSafeReload(this, Filename);	
 	if (Branch == nullptr)
@@ -4835,7 +4920,6 @@ bool FConfigCacheIni::GetText( const TCHAR* Section, const TCHAR* Key, FText& Va
 	
 	if (Branch->bIsSafeUnloaded)
 	{
-		UE_LOG(LogConfig, Log, TEXT("About to reload due to getting %s: %s"), Section, Key);
 		Branch->SafeReload();
 	}
 #else
@@ -4871,6 +4955,8 @@ bool FConfigCacheIni::GetText( const TCHAR* Section, const TCHAR* Key, FText& Va
 
 bool FConfigCacheIni::GetSection( const TCHAR* Section, TArray<FString>& Result, const FString& Filename )
 {
+	FConfigLogContextScope LogScope(TEXT("GetSection"), Section, nullptr);
+
 	FRemoteConfig::Get()->FinishRead(*Filename); // Ensure the remote file has been loaded and processed
 	Result.Reset();
 	FConfigFile* File = Find(Filename);
@@ -4896,6 +4982,8 @@ bool FConfigCacheIni::GetSection( const TCHAR* Section, TArray<FString>& Result,
 
 const FConfigSection* FConfigCacheIni::GetSection( const TCHAR* Section, const bool Force, const FString& Filename )
 {
+	FConfigLogContextScope LogScope(TEXT("GetSection"), Section, nullptr);
+
 	FRemoteConfig::Get()->FinishRead(*Filename); // Ensure the remote file has been loaded and processed
 	FConfigFile* File = Find(Filename);
 	if (!File)
@@ -4924,6 +5012,8 @@ const FConfigSection* FConfigCacheIni::GetSection( const TCHAR* Section, const b
 
 bool FConfigCacheIni::DoesSectionExist(const TCHAR* Section, const FString& Filename)
 {
+	FConfigLogContextScope LogScope(TEXT("DoesSectionExist"), Section, nullptr);
+
 	bool bReturnVal = false;
 
 	FRemoteConfig::Get()->FinishRead(*Filename); // Ensure the remote file has been loaded and processed
@@ -5106,6 +5196,8 @@ void FConfigCacheIni::GetConfigFilenames(TArray<FString>& ConfigFilenames)
  */
 bool FConfigCacheIni::GetSectionNames( const FString& Filename, TArray<FString>& out_SectionNames )
 {
+	FConfigLogContextScope LogScope(TEXT("GetSectionNames"), nullptr, nullptr);
+
 	bool bResult = false;
 
 	FConfigFile* File = Find(Filename);
@@ -5136,6 +5228,8 @@ bool FConfigCacheIni::GetSectionNames( const FString& Filename, TArray<FString>&
  */
 bool FConfigCacheIni::GetPerObjectConfigSections( const FString& Filename, const FString& SearchClass, TArray<FString>& out_SectionNames, int32 MaxResults )
 {
+	FConfigLogContextScope LogScope(TEXT("GetPerObjectConfigSections"), *SearchClass, nullptr);
+
 	bool bResult = false;
 
 	MaxResults = FMath::Max(0, MaxResults);
@@ -6803,6 +6897,8 @@ bool FConfigCacheIni::LoadExternalIniFile(FConfigFile & ConfigFile, const TCHAR
 
 FConfigFile* FConfigCacheIni::FindPlatformConfig(const TCHAR* IniName, const TCHAR* Platform)
 {
+	FConfigLogContextScope LogScope(TEXT("FindPlatformConfig"), nullptr, nullptr);
+
 	if (Platform != nullptr && FCString::Stricmp(Platform, ANSI_TO_TCHAR(FPlatformProperties::IniPlatformName())) != 0)
 	{
 #if ALLOW_OTHER_PLATFORM_CONFIG
@@ -6822,6 +6918,8 @@ FConfigFile* FConfigCacheIni::FindPlatformConfig(const TCHAR* IniName, const TCH
 
 FConfigFile* FConfigCacheIni::FindOrLoadPlatformConfig(FConfigFile& LocalFile, const TCHAR* IniName, const TCHAR* Platform)
 {
+	FConfigLogContextScope LogScope(TEXT("FindOrLoadPlatformConfig"), nullptr, nullptr);
+
 	FConfigFile* File = FindPlatformConfig(IniName, Platform);
 	if (File == nullptr)
 	{
