[WinArm64] Enabling callstacks when x64 is run on ARM64 via emulation.

#rnx

#Synced-CL 51108079

diff --git a/Engine/Source/Runtime/Core/Private/ProfilingDebugging/Microsoft/WindowsCallstackTrace.cpp b/Engine/Source/Runtime/Core/Private/ProfilingDebugging/Microsoft/WindowsCallstackTrace.cpp
--- a/Engine/Source/Runtime/Core/Private/ProfilingDebugging/Microsoft/WindowsCallstackTrace.cpp
+++ b/Engine/Source/Runtime/Core/Private/ProfilingDebugging/Microsoft/WindowsCallstackTrace.cpp
@@ -996,13 +996,6 @@ void Modules_Initialize();
 ////////////////////////////////////////////////////////////////////////////////
 void CallstackTrace_CreateInternal(FMalloc* Malloc)
 {
-	// backtrace doesn't work on arm64, even with x64 emulation
-#if PLATFORM_CPU_X86_FAMILY
-	if (FPlatformMisc::GetHostArchitecture() == TEXT("arm64"))
-	{
-		return;
-	}
-#endif
 
 	if (FBacktracer::Get() != nullptr)
 	{
