Fix CIS issue in IOS

A previous CL moved a typedef for RECT into its own header file, but one of the files that was manually defining it was left out of the initial commit by accident.

#Synced-CL 51141141

diff --git a/Engine/Source/Runtime/Core/Public/IOS/IOSSystemIncludes.h b/Engine/Source/Runtime/Core/Public/IOS/IOSSystemIncludes.h
--- a/Engine/Source/Runtime/Core/Public/IOS/IOSSystemIncludes.h
+++ b/Engine/Source/Runtime/Core/Public/IOS/IOSSystemIncludes.h
@@ -63,14 +63,7 @@
 #include <xmmintrin.h>
 #endif
 
-struct tagRECT
-{
-	int32 left;
-	int32 top;
-	int32 right;
-	int32 bottom;
-};
-typedef struct tagRECT RECT;
+#include "Apple/AppleTagRect.h"
 
 /*----------------------------------------------------------------------------
  Memory. On Mac OS X malloc allocates memory aligned to 16 bytes.
