[Core/ProfilingDebugging] CallstackTrace: Removed redundant counter in FCallstackXORAndRLE::Compress. Minor code style fixes in the rest of the file.

#Synced-CL 51157224

diff --git a/Engine/Source/Runtime/Core/Public/ProfilingDebugging/CallstackTrace.h b/Engine/Source/Runtime/Core/Public/ProfilingDebugging/CallstackTrace.h
--- a/Engine/Source/Runtime/Core/Public/ProfilingDebugging/CallstackTrace.h
+++ b/Engine/Source/Runtime/Core/Public/ProfilingDebugging/CallstackTrace.h
@@ -1,4 +1,5 @@
 // Copyright Epic Games, Inc. All Rights Reserved.
+
 #pragma once
 
 #include "AutoRTFM/Defines.h"
@@ -58,10 +59,11 @@ AUTORTFM_OPEN CORE_API uint32 CallstackTrace_GetCurrentId();
  /**
   * Callstack Trace Scoped Macro to avoid resolving the full callstack
   * can be used when some external libraries are not compiled with frame pointers
-  * preventing us to resolve it without crashing. Instead the callstack will be 
+  * preventing us to resolve it without crashing. Instead the callstack will be
   * only the caller address.
   */
-#define CALLSTACK_TRACE_LIMIT_CALLSTACKRESOLVE_SCOPE() 						FCallStackTraceLimitResolveScope UE_JOIN(FCTLMScope,__LINE__)
+#define CALLSTACK_TRACE_LIMIT_CALLSTACKRESOLVE_SCOPE() \
+	FCallStackTraceLimitResolveScope UE_JOIN(FCTLMScope,__LINE__)
 
 extern CORE_API uint32 GCallStackTracingTlsSlotIndex;
 
@@ -71,9 +73,13 @@ extern CORE_API uint32 GCallStackTracingTlsSlotIndex;
 inline void* CallstackTrace_GetFallbackPlatformReturnAddressData()
 {
 	if (FPlatformTLS::IsValidTlsSlot(GCallStackTracingTlsSlotIndex))
+	{
 		return FPlatformTLS::GetTlsValue(GCallStackTracingTlsSlotIndex);
+	}
 	else
+	{
 		return nullptr;
+	}
 }
 
 /**
@@ -129,7 +135,6 @@ namespace FCallstackXORAndRLE
 	inline uint32 Compress(TConstArrayView<uint64> Addresses, TArrayView<uint8> Compressed)
 	{
 		check(Addresses.Num() * 8 <= Compressed.Num());
-		uint32 OutIdx = 0;
 		uint32 CompressedBytes = 0;
 		uint64 LastAddress = 0;
 		for(uint8 Index = 0; Index < Addresses.Num(); ++Index)
@@ -138,18 +143,16 @@ namespace FCallstackXORAndRLE
 			uint64 XorAddress = LastAddress ^ Address;
 			LastAddress = Address;
 			const uint8 LeadingZeroes = uint8(FMath::CountLeadingZeros64(XorAddress));
-			Compressed[OutIdx++] = LeadingZeroes;
-			++CompressedBytes;
+			Compressed[CompressedBytes++] = LeadingZeroes;
 			if (LeadingZeroes < 64)
 			{
 				int8 ValueBits = 64 - LeadingZeroes;
 				do
 				{
-					Compressed[OutIdx++] = uint8(0xff & XorAddress);
+					Compressed[CompressedBytes++] = uint8(0xff & XorAddress);
 					XorAddress >>= 8;
 					ValueBits -= 8;
-					++CompressedBytes;
-				} while(ValueBits > 0);  
+				} while(ValueBits > 0);
 			}
 		}
 		return CompressedBytes;
@@ -181,11 +184,11 @@ namespace FCallstackXORAndRLE
 
 			if (OutTotalAddressCount < (uint32)Addresses.Num())
 			{
-			    Addresses[OutTotalAddressCount] = OutAddress;
+				Addresses[OutTotalAddressCount] = OutAddress;
 			}
 			OutTotalAddressCount++;
 		}
-		return FMath::Min(OutTotalAddressCount,  (uint32)Addresses.Num());
+		return FMath::Min(OutTotalAddressCount, (uint32)Addresses.Num());
 	}
 } // namespace FCallstackXORAndRLE
 
@@ -234,18 +237,18 @@ namespace FCallstackDelta7bit
 				EncodedAddress |= uint64(Byte & 0x7F) << Shift;
 				Shift += 7;
 			} while (Byte & 0x80);
-			
+
 			int64 Delta = (EncodedAddress >> 1) ^ -(int64)(EncodedAddress & 1);
 			uint64 Address = LastAddress + Delta;
 			LastAddress = Address;
-			
+
 			if (OutTotalAddressCount < (uint32)Addresses.Num())
 			{
-			    Addresses[OutTotalAddressCount] = Address;
+				Addresses[OutTotalAddressCount] = Address;
 			}
 			OutTotalAddressCount++;
 		}
-		return FMath::Min(OutTotalAddressCount,  (uint32)Addresses.Num());
+		return FMath::Min(OutTotalAddressCount, (uint32)Addresses.Num());
 	}
 } // namespace FCallstackDelta7bit
 
@@ -286,10 +289,10 @@ namespace FCallstackDeltaVarInt
 
 			if (OutTotalAddressCount < (uint32)Addresses.Num())
 			{
-			    Addresses[OutTotalAddressCount] = Address;
+				Addresses[OutTotalAddressCount] = Address;
 			}
 			OutTotalAddressCount++;
 		}
-		return FMath::Min(OutTotalAddressCount,  (uint32)Addresses.Num());
+		return FMath::Min(OutTotalAddressCount, (uint32)Addresses.Num());
 	}
-} // namespace FCallstackDeltaVarInt
\ No newline at end of file
+} // namespace FCallstackDeltaVarInt
