Force immediate exit when FCrashReportingThread or FRunnableThreadWin crashes on Windows.
- If we are already in the process of exiting, calling ::exit() will invoke crt exit table which might already be running on the main thread. This may lead to deadlocks between threads or cause other issues.
- Fixes a case where process was left running indefinitely after closing the game due to deadlock from crt exit table invoked twice.

[RN] rnx

#Synced-CL 51088011

diff --git a/Engine/Source/Runtime/Core/Private/Windows/WindowsPlatformCrashContext.cpp b/Engine/Source/Runtime/Core/Private/Windows/WindowsPlatformCrashContext.cpp
--- a/Engine/Source/Runtime/Core/Private/Windows/WindowsPlatformCrashContext.cpp
+++ b/Engine/Source/Runtime/Core/Private/Windows/WindowsPlatformCrashContext.cpp
@@ -1325,7 +1325,7 @@ private:
 		__except(EXCEPTION_EXECUTE_HANDLER)
 		{
 			// The crash reporting thread crashed itself. Exit with a code that the out-of-process monitor will be able to pick up and report into analytics.
-			::exit(ECrashExitCodes::CrashReporterCrashed);
+			TerminateProcess(GetCurrentProcess(), (UINT)ECrashExitCodes::CrashReporterCrashed);
 		}
 #endif
 		return 0;
diff --git a/Engine/Source/Runtime/Core/Private/Windows/WindowsRunnableThread.cpp b/Engine/Source/Runtime/Core/Private/Windows/WindowsRunnableThread.cpp
--- a/Engine/Source/Runtime/Core/Private/Windows/WindowsRunnableThread.cpp
+++ b/Engine/Source/Runtime/Core/Private/Windows/WindowsRunnableThread.cpp
@@ -109,8 +109,8 @@ uint32 FRunnableThreadWin::GuardedRun()
 				// The crash handler crashed itself, exit with a code which the 
 				// out-of-process monitor will be able to pick up and report into 
 				// analytics.
-
-				::exit(ECrashExitCodes::CrashHandlerCrashed);
+				
+				TerminateProcess(GetCurrentProcess(), (UINT)ECrashExitCodes::CrashHandlerCrashed);
 			}
 		}
 #endif // !PLATFORM_SEH_EXCEPTIONS_DISABLED
