FGuid: Create UUIDv7 (time-and-random) on Unix/Linux & Android

Prior algorithm (UUID v4) is still available via FUnixPlatformGuid::CreateGuidFromRandomOnly().
UUIDv7 is also available via FUnixPlatformGuid::CreateGuidFromTimeAndRandom().

Removed ancient support for pre-glibc-2.25 getrandom syscall.

[RN] CreateGuid algorithm changed to UUIDv7 on Unix & Android
#tests LinuxServer FGuids inspected with gdb & `uuid -d`.

#Synced-CL 50901758

diff --git a/Engine/Source/Runtime/Core/Private/Android/AndroidPlatformMisc.cpp b/Engine/Source/Runtime/Core/Private/Android/AndroidPlatformMisc.cpp
--- a/Engine/Source/Runtime/Core/Private/Android/AndroidPlatformMisc.cpp
+++ b/Engine/Source/Runtime/Core/Private/Android/AndroidPlatformMisc.cpp
@@ -25,6 +25,7 @@
 #include "HAL/PlatformMallocCrash.h"
 #include "Android/AndroidJavaMessageBox.h"
 #include "GenericPlatform/GenericPlatformChunkInstall.h"
+#include "Unix/UnixPlatformGuid.h"
 
 #include "Misc/Parse.h"
 #include "Internationalization/Regex.h"
@@ -1103,81 +1104,26 @@ bool FAndroidMisc::SupportsLocalCaching()
 
 }
 
-static int SysGetRandomSupported = -1;
-
-// http://man7.org/linux/man-pages/man2/getrandom.2.html
-// getrandom() was introduced in version 3.17 of the Linux kernel
-//   and glibc version 2.25.
-
-// Check known platforms if SYS_getrandom isn't defined
-#if !defined(SYS_getrandom)
-	#if PLATFORM_CPU_X86_FAMILY
-		#define SYS_getrandom 318
-	#elif PLATFORM_CPU_ARM_FAMILY
-		#define SYS_getrandom 278
-	#endif
-#endif // !defined(SYS_getrandom)
-
-namespace
-{
-#if defined(SYS_getrandom)
-
-#if !defined(GRND_NONBLOCK)
-	#define GRND_NONBLOCK 0x0001
-#endif
-
-	int SysGetRandom(void *buf, size_t buflen)
-	{
-		if (SysGetRandomSupported < 0)
-		{
-			int Ret = syscall(SYS_getrandom, buf, buflen, GRND_NONBLOCK);
-
-			// If -1 is returned with ENOSYS, kernel doesn't support getrandom
-			SysGetRandomSupported = ((Ret == -1) && (errno == ENOSYS)) ? 0 : 1;
-		}
-
-		return SysGetRandomSupported ?
-			syscall(SYS_getrandom, buf, buflen, GRND_NONBLOCK) : -1;
-	}
-
-#else
-
-	int SysGetRandom(void *buf, size_t buflen)
-	{
-		return -1;
-	}
-
-#endif // !SYS_getrandom
-}
-
 /**
- * Try to use SYS_getrandom which would be the fastest, otherwise fall back to 
+ * Try to use the Unix/Linux method, clock_gettime() & getrandom(), which would be the fastest, otherwise fall back to 
  * use /proc/sys/kernel/random/uuid to get GUID; do NOT use JNI since this may be called too early
  */
-void  FAndroidMisc::CreateGuid(struct FGuid& Result)
+void FAndroidMisc::CreateGuid(struct FGuid& Result)
 {
 	QUICK_SCOPE_CYCLE_COUNTER(STAT_FGenericPlatformMisc_CreateGuid);
 
-	static bool bGetRandomFailed = false;
+	static bool bUnixMethodOK = true;
 	static bool bProcUUIDFailed = false;
 
-	if (!bGetRandomFailed)
+	if (bUnixMethodOK) [[likely]]
 	{
-		int BytesRead = SysGetRandom(&Result, sizeof(Result));
-
-		if (BytesRead == sizeof(Result))
+		Result = FUnixPlatformGuid::CreateGuidFromTimeAndRandom();
+		if (Result.IsValid()) [[likely]]
 		{
-			// https://tools.ietf.org/html/rfc4122#section-4.4
-			// https://en.wikipedia.org/wiki/Universally_unique_identifier
-			//
-			// The 4 bits of digit M indicate the UUID version, and the 1â€“3
-			//   most significant bits of digit N indicate the UUID variant.
-			// xxxxxxxx-xxxx-Mxxx-Nxxx-xxxxxxxxxxxx
-			Result[1] = (Result[1] & 0xffff0fff) | 0x00004000; // version 4
-			Result[2] = (Result[2] & 0x3fffffff) | 0x80000000; // variant 1
 			return;
 		}
-		bGetRandomFailed = true;
+
+		bUnixMethodOK = false;
 	}
 
 #define FROM_HEX(_a) ( (_a) <= '9' ? (_a) - '0' : (_a) <= 'F' ?  (_a) - 'A' + 10 : (_a) - 'a' + 10 )
diff --git a/Engine/Source/Runtime/Core/Private/Unix/UnixPlatformGuid.cpp b/Engine/Source/Runtime/Core/Private/Unix/UnixPlatformGuid.cpp
new file mode 100644
--- /dev/null
+++ b/Engine/Source/Runtime/Core/Private/Unix/UnixPlatformGuid.cpp
@@ -0,0 +1,76 @@
+// Copyright Epic Games, Inc. All Rights Reserved.
+
+#include "Unix/UnixPlatformGuid.h"
+
+#include "HAL/Platform.h"
+#include "Misc/AssertionMacros.h"
+
+#include <sys/random.h>
+#include <time.h>
+
+namespace
+{
+	// Set the 4-bit Version and 2-bit Variant fields of a FGuid.
+	static inline void SetGuidType(FGuid& Target, uint8 Version, uint8 Variant)
+	{
+		Target.B = (Target.B & 0xffff0fff) | ((static_cast<uint32>(Version) & 0b1111) << 12);
+		Target.C = (Target.C & 0x3fffffff) | ((static_cast<uint32>(Variant) & 0b11) << 30);
+	}
+
+} // namespace (anonymous)
+
+FGuid FUnixPlatformGuid::CreateGuid()
+{
+	return CreateGuidFromTimeAndRandom();
+}
+
+FGuid FUnixPlatformGuid::CreateGuidFromTimeAndRandom()
+{
+	/**
+	 * UUIDv7 format:
+	 * - 48 bits Unix timestamp in milliseconds
+	 * -  4 bits version (7)
+	 * - 12 bits random data
+	 * -  2 bits variant (0b10)
+	 * - 62 bits random data
+	 */
+	struct timespec ts = {};
+	if (!ensure(clock_gettime(CLOCK_REALTIME, &ts) == 0))
+	{
+		return FGuid();
+	}
+
+	const uint64 UnixTimestampMilliseconds = static_cast<uint64>(ts.tv_sec) * 1000 + ts.tv_nsec / 1'000'000;
+
+	FGuid Guid = {};
+	// Load random data from 2nd half of B through D (80 bits, although 6 bits will be overwritten with the version & variant fields).
+	constexpr size_t NumRandomBytes = 10;
+	constexpr unsigned int GetRandomFlags = 0;
+	if (!ensure(getrandom(reinterpret_cast<char*>(&Guid.A) + 6, NumRandomBytes, GetRandomFlags) == NumRandomBytes))
+	{
+		return FGuid();
+	}
+
+#if PLATFORM_LITTLE_ENDIAN
+	// Shift random half of Guid.B into the lower end.
+	Guid.B >>= 16;
+#endif
+
+	Guid.A = static_cast<uint32>(UnixTimestampMilliseconds >> 16);
+	Guid.B |= static_cast<uint32>(UnixTimestampMilliseconds) << 16;
+	SetGuidType(Guid, 7, 0b10);
+	return Guid;
+}
+
+FGuid FUnixPlatformGuid::CreateGuidFromRandomOnly()
+{
+	FGuid Guid = {};
+	constexpr unsigned int GetRandomFlags = 0;
+	if (!ensure(getrandom(&Guid, sizeof(Guid), GetRandomFlags) == sizeof(Guid)))
+	{
+		return FGuid();
+	}
+
+	SetGuidType(Guid, 4, 0b10);
+	return Guid;
+}
diff --git a/Engine/Source/Runtime/Core/Private/Unix/UnixPlatformMisc.cpp b/Engine/Source/Runtime/Core/Private/Unix/UnixPlatformMisc.cpp
--- a/Engine/Source/Runtime/Core/Private/Unix/UnixPlatformMisc.cpp
+++ b/Engine/Source/Runtime/Core/Private/Unix/UnixPlatformMisc.cpp
@@ -22,6 +22,7 @@
 #include "Misc/App.h"
 #include "Misc/FileHelper.h"
 #include "Unix/UnixPlatformCrashContext.h"
+#include "Unix/UnixPlatformGuid.h"
 #include "Unix/UnixPlatformSyscallTable.h"
 #include "Misc/ConfigCacheIni.h"
 #include "GenericPlatform/GenericPlatformChunkInstall.h"
@@ -56,8 +57,6 @@
 
 extern bool GInitializedSDL;
 
-static int SysGetRandomSupported = -1;
-
 namespace PlatformMiscLimits
 {
 	enum
@@ -243,12 +242,6 @@ void FUnixPlatformMisc::PlatformInit()
 		setvbuf(stdout, NULL, _IONBF, 0);
 	}
 
-	if (FParse::Param(FCommandLine::Get(), TEXT("norandomguids")))
-	{
-		// If "-norandomguids" specified, don't use SYS_getrandom syscall
-		SysGetRandomSupported = 0;
-	}
-
 	// This symbol is used for debugging but with LTO enabled it gets stripped as nothing is using it
 	// Lets use it here just to log if its valid under VeryVerbose settings
 	extern uint8** GNameBlocksDebug;
@@ -1629,96 +1622,9 @@ bool FUnixPlatformMisc::IsRunningOnBattery()
 	return bIsOnBattery;
 }
 
-#if PLATFORM_UNIX && defined(_GNU_SOURCE)
-
-#include <sys/syscall.h>
-
-// http://man7.org/linux/man-pages/man2/getrandom.2.html
-// getrandom() was introduced in version 3.17 of the Linux kernel
-//   and glibc version 2.25.
-
-// Check known platforms if SYS_getrandom isn't defined
-#if !defined(SYS_getrandom)
-	#if PLATFORM_CPU_X86_FAMILY
-		#define SYS_getrandom 318
-	#elif PLATFORM_CPU_ARM_FAMILY
-		#define SYS_getrandom 278
-	#endif
-#endif // !defined(SYS_getrandom)
-
-#endif // PLATFORM_UNIX && _GNU_SOURCE
-
-namespace
-{
-#if defined(SYS_getrandom)
-
-#if !defined(GRND_NONBLOCK)
-	#define GRND_NONBLOCK 0x0001
-#endif
-
-	int SysGetRandom(void *buf, size_t buflen)
-	{
-		if (SysGetRandomSupported < 0)
-		{
-			int Ret = syscall(SYS_getrandom, buf, buflen, GRND_NONBLOCK);
-
-			// If -1 is returned with ENOSYS, kernel doesn't support getrandom
-			SysGetRandomSupported = ((Ret == -1) && (errno == ENOSYS)) ? 0 : 1;
-		}
-
-		return SysGetRandomSupported ?
-			syscall(SYS_getrandom, buf, buflen, GRND_NONBLOCK) : -1;
-	}
-
-#else
-
-	int SysGetRandom(void *buf, size_t buflen)
-	{
-		return -1;
-	}
-
-#endif // !SYS_getrandom
-
-	void CreateGuidImpl(FGuid& Result)
-	{
-		int BytesRead = SysGetRandom(&Result, sizeof(Result));
-
-		if (BytesRead == sizeof(Result))
-		{
-			// https://tools.ietf.org/html/rfc4122#section-4.4
-			// https://en.wikipedia.org/wiki/Universally_unique_identifier
-			//
-			// The 4 bits of digit M indicate the UUID version, and the 1–3
-			//   most significant bits of digit N indicate the UUID variant.
-			// xxxxxxxx-xxxx-Mxxx-Nxxx-xxxxxxxxxxxx
-			Result[1] = (Result[1] & 0xffff0fff) | 0x00004000; // version 4
-			Result[2] = (Result[2] & 0x3fffffff) | 0x80000000; // variant 1
-		}
-		else
-		{
-			// Fall back to generic CreateGuid
-			FGenericPlatformMisc::CreateGuid(Result);
-		}
-	}
-}
-
-// If we fail to create a Guid with urandom fallback to the generic platform.
-// This maybe need to be tweaked for Servers and hard fail here
 void FUnixPlatformMisc::CreateGuid(FGuid& Result)
 {
-	if (AutoRTFM::IsTransactional())
-	{
-		Result = AutoRTFM::Open([]
-		{
-			FGuid Guid;
-			CreateGuidImpl(Guid);
-			return Guid;
-		});
-	}
-	else
-	{
-		CreateGuidImpl(Result);
-	}
+	Result = FUnixPlatformGuid::CreateGuid();
 }
 
 CORE_API TFunction<void()> UngrabAllInputCallback;
diff --git a/Engine/Source/Runtime/Core/Public/Unix/UnixPlatformGuid.h b/Engine/Source/Runtime/Core/Public/Unix/UnixPlatformGuid.h
new file mode 100644
--- /dev/null
+++ b/Engine/Source/Runtime/Core/Public/Unix/UnixPlatformGuid.h
@@ -0,0 +1,33 @@
+// Copyright Epic Games, Inc. All Rights Reserved.
+#pragma once
+
+#include "Misc/Guid.h"
+
+/**
+ * Functions for operating on GUIDs (UUIDs).
+ */
+struct FUnixPlatformGuid
+{
+	/**
+	 * Create a GUID (UUID). The UUID will be generated by one of the standard algorithms in RFC 9562.
+	 * Historically this function used the random-data-only algorithm (UUID v4) until Unreal Engine 5.8, where it switched to time-and-random-data
+	 * (UUID v7).
+	 *
+	 * @return New GUID, or Invalid() (Nil UUID—all zeroes) on error.
+	 */
+	AUTORTFM_OPEN static CORE_API FGuid CreateGuid();
+
+	/**
+	 * Create a GUID (UUID) using the time-and-random-data method (UUID version 7).
+	 *
+	 * @return New GUID, or Invalid() (Nil UUID—all zeroes) on error.
+	 */
+	AUTORTFM_OPEN static CORE_API FGuid CreateGuidFromTimeAndRandom();
+
+	/**
+	 * Create a GUID (UUID) using the random data only method (UUID version 4).
+	 *
+	 * @return New GUID, or Invalid() (Nil UUID—all zeroes) on error.
+	 */
+	AUTORTFM_OPEN static CORE_API FGuid CreateGuidFromRandomOnly();
+};
