[Backout] - CL51170242
#rnx
Original CL Desc
-----------------------------------------------------------------
FGuid: Create UUIDv7 (time-and-random) on Unix/Linux & Android

Prior algorithm (UUID v4) is still available via FUnixPlatformGuid::CreateGuidFromRandomOnly().
UUIDv7 is also available via FUnixPlatformGuid::CreateGuidFromTimeAndRandom().

Removed ancient support for pre-glibc-2.25 getrandom syscall.

[RN] CreateGuid algorithm changed to UUIDv7 on Unix & Android
#tests LinuxServer FGuids inspected with gdb & `uuid -d`.

diff --git a/Engine/Source/Runtime/Core/Private/Android/AndroidPlatformMisc.cpp b/Engine/Source/Runtime/Core/Private/Android/AndroidPlatformMisc.cpp
--- a/Engine/Source/Runtime/Core/Private/Android/AndroidPlatformMisc.cpp
+++ b/Engine/Source/Runtime/Core/Private/Android/AndroidPlatformMisc.cpp
@@ -25,7 +25,6 @@
 #include "HAL/PlatformMallocCrash.h"
 #include "Android/AndroidJavaMessageBox.h"
 #include "GenericPlatform/GenericPlatformChunkInstall.h"
-#include "Unix/UnixPlatformGuid.h"
 
 #include "Misc/Parse.h"
 #include "Internationalization/Regex.h"
@@ -1104,26 +1103,81 @@ bool FAndroidMisc::SupportsLocalCaching()
 
 }
 
+static int SysGetRandomSupported = -1;
+
+// http://man7.org/linux/man-pages/man2/getrandom.2.html
+// getrandom() was introduced in version 3.17 of the Linux kernel
+//   and glibc version 2.25.
+
+// Check known platforms if SYS_getrandom isn't defined
+#if !defined(SYS_getrandom)
+	#if PLATFORM_CPU_X86_FAMILY
+		#define SYS_getrandom 318
+	#elif PLATFORM_CPU_ARM_FAMILY
+		#define SYS_getrandom 278
+	#endif
+#endif // !defined(SYS_getrandom)
+
+namespace
+{
+#if defined(SYS_getrandom)
+
+#if !defined(GRND_NONBLOCK)
+	#define GRND_NONBLOCK 0x0001
+#endif
+
+	int SysGetRandom(void *buf, size_t buflen)
+	{
+		if (SysGetRandomSupported < 0)
+		{
+			int Ret = syscall(SYS_getrandom, buf, buflen, GRND_NONBLOCK);
+
+			// If -1 is returned with ENOSYS, kernel doesn't support getrandom
+			SysGetRandomSupported = ((Ret == -1) && (errno == ENOSYS)) ? 0 : 1;
+		}
+
+		return SysGetRandomSupported ?
+			syscall(SYS_getrandom, buf, buflen, GRND_NONBLOCK) : -1;
+	}
+
+#else
+
+	int SysGetRandom(void *buf, size_t buflen)
+	{
+		return -1;
+	}
+
+#endif // !SYS_getrandom
+}
+
 /**
- * Try to use the Unix/Linux method, clock_gettime() & getrandom(), which would be the fastest, otherwise fall back to 
+ * Try to use SYS_getrandom which would be the fastest, otherwise fall back to 
  * use /proc/sys/kernel/random/uuid to get GUID; do NOT use JNI since this may be called too early
  */
-void FAndroidMisc::CreateGuid(struct FGuid& Result)
+void  FAndroidMisc::CreateGuid(struct FGuid& Result)
 {
 	QUICK_SCOPE_CYCLE_COUNTER(STAT_FGenericPlatformMisc_CreateGuid);
 
-	static bool bUnixMethodOK = true;
+	static bool bGetRandomFailed = false;
 	static bool bProcUUIDFailed = false;
 
-	if (bUnixMethodOK) [[likely]]
+	if (!bGetRandomFailed)
 	{
-		Result = FUnixPlatformGuid::CreateGuidFromTimeAndRandom();
-		if (Result.IsValid()) [[likely]]
+		int BytesRead = SysGetRandom(&Result, sizeof(Result));
+
+		if (BytesRead == sizeof(Result))
 		{
+			// https://tools.ietf.org/html/rfc4122#section-4.4
+			// https://en.wikipedia.org/wiki/Universally_unique_identifier
+			//
+			// The 4 bits of digit M indicate the UUID version, and the 1â€“3
+			//   most significant bits of digit N indicate the UUID variant.
+			// xxxxxxxx-xxxx-Mxxx-Nxxx-xxxxxxxxxxxx
+			Result[1] = (Result[1] & 0xffff0fff) | 0x00004000; // version 4
+			Result[2] = (Result[2] & 0x3fffffff) | 0x80000000; // variant 1
 			return;
 		}
-
-		bUnixMethodOK = false;
+		bGetRandomFailed = true;
 	}
 
 #define FROM_HEX(_a) ( (_a) <= '9' ? (_a) - '0' : (_a) <= 'F' ?  (_a) - 'A' + 10 : (_a) - 'a' + 10 )
diff --git a/Engine/Source/Runtime/Core/Private/Unix/UnixPlatformGuid.cpp b/Engine/Source/Runtime/Core/Private/Unix/UnixPlatformGuid.cpp
deleted file mode 100644
diff --git a/Engine/Source/Runtime/Core/Private/Unix/UnixPlatformMisc.cpp b/Engine/Source/Runtime/Core/Private/Unix/UnixPlatformMisc.cpp
--- a/Engine/Source/Runtime/Core/Private/Unix/UnixPlatformMisc.cpp
+++ b/Engine/Source/Runtime/Core/Private/Unix/UnixPlatformMisc.cpp
@@ -22,7 +22,6 @@
 #include "Misc/App.h"
 #include "Misc/FileHelper.h"
 #include "Unix/UnixPlatformCrashContext.h"
-#include "Unix/UnixPlatformGuid.h"
 #include "Unix/UnixPlatformSyscallTable.h"
 #include "Misc/ConfigCacheIni.h"
 #include "GenericPlatform/GenericPlatformChunkInstall.h"
@@ -57,6 +56,8 @@
 
 extern bool GInitializedSDL;
 
+static int SysGetRandomSupported = -1;
+
 namespace PlatformMiscLimits
 {
 	enum
@@ -242,6 +243,12 @@ void FUnixPlatformMisc::PlatformInit()
 		setvbuf(stdout, NULL, _IONBF, 0);
 	}
 
+	if (FParse::Param(FCommandLine::Get(), TEXT("norandomguids")))
+	{
+		// If "-norandomguids" specified, don't use SYS_getrandom syscall
+		SysGetRandomSupported = 0;
+	}
+
 	// This symbol is used for debugging but with LTO enabled it gets stripped as nothing is using it
 	// Lets use it here just to log if its valid under VeryVerbose settings
 	extern uint8** GNameBlocksDebug;
@@ -1622,9 +1629,96 @@ bool FUnixPlatformMisc::IsRunningOnBattery()
 	return bIsOnBattery;
 }
 
+#if PLATFORM_UNIX && defined(_GNU_SOURCE)
+
+#include <sys/syscall.h>
+
+// http://man7.org/linux/man-pages/man2/getrandom.2.html
+// getrandom() was introduced in version 3.17 of the Linux kernel
+//   and glibc version 2.25.
+
+// Check known platforms if SYS_getrandom isn't defined
+#if !defined(SYS_getrandom)
+	#if PLATFORM_CPU_X86_FAMILY
+		#define SYS_getrandom 318
+	#elif PLATFORM_CPU_ARM_FAMILY
+		#define SYS_getrandom 278
+	#endif
+#endif // !defined(SYS_getrandom)
+
+#endif // PLATFORM_UNIX && _GNU_SOURCE
+
+namespace
+{
+#if defined(SYS_getrandom)
+
+#if !defined(GRND_NONBLOCK)
+	#define GRND_NONBLOCK 0x0001
+#endif
+
+	int SysGetRandom(void *buf, size_t buflen)
+	{
+		if (SysGetRandomSupported < 0)
+		{
+			int Ret = syscall(SYS_getrandom, buf, buflen, GRND_NONBLOCK);
+
+			// If -1 is returned with ENOSYS, kernel doesn't support getrandom
+			SysGetRandomSupported = ((Ret == -1) && (errno == ENOSYS)) ? 0 : 1;
+		}
+
+		return SysGetRandomSupported ?
+			syscall(SYS_getrandom, buf, buflen, GRND_NONBLOCK) : -1;
+	}
+
+#else
+
+	int SysGetRandom(void *buf, size_t buflen)
+	{
+		return -1;
+	}
+
+#endif // !SYS_getrandom
+
+	void CreateGuidImpl(FGuid& Result)
+	{
+		int BytesRead = SysGetRandom(&Result, sizeof(Result));
+
+		if (BytesRead == sizeof(Result))
+		{
+			// https://tools.ietf.org/html/rfc4122#section-4.4
+			// https://en.wikipedia.org/wiki/Universally_unique_identifier
+			//
+			// The 4 bits of digit M indicate the UUID version, and the 1–3
+			//   most significant bits of digit N indicate the UUID variant.
+			// xxxxxxxx-xxxx-Mxxx-Nxxx-xxxxxxxxxxxx
+			Result[1] = (Result[1] & 0xffff0fff) | 0x00004000; // version 4
+			Result[2] = (Result[2] & 0x3fffffff) | 0x80000000; // variant 1
+		}
+		else
+		{
+			// Fall back to generic CreateGuid
+			FGenericPlatformMisc::CreateGuid(Result);
+		}
+	}
+}
+
+// If we fail to create a Guid with urandom fallback to the generic platform.
+// This maybe need to be tweaked for Servers and hard fail here
 void FUnixPlatformMisc::CreateGuid(FGuid& Result)
 {
-	Result = FUnixPlatformGuid::CreateGuid();
+	if (AutoRTFM::IsTransactional())
+	{
+		Result = AutoRTFM::Open([]
+		{
+			FGuid Guid;
+			CreateGuidImpl(Guid);
+			return Guid;
+		});
+	}
+	else
+	{
+		CreateGuidImpl(Result);
+	}
 }
 
 CORE_API TFunction<void()> UngrabAllInputCallback;
diff --git a/Engine/Source/Runtime/Core/Public/Unix/UnixPlatformGuid.h b/Engine/Source/Runtime/Core/Public/Unix/UnixPlatformGuid.h
deleted file mode 100644
