Fix for CriticallyDampedSmoothing and add some SpringDamper tests

#Synced-CL 51168715

diff --git a/Engine/Source/Runtime/Core/Public/Math/UnrealMathUtility.h b/Engine/Source/Runtime/Core/Public/Math/UnrealMathUtility.h
--- a/Engine/Source/Runtime/Core/Public/Math/UnrealMathUtility.h
+++ b/Engine/Source/Runtime/Core/Public/Math/UnrealMathUtility.h
@@ -1627,7 +1627,7 @@ public:
 			//
 			// We're damping towards a target, so convert - i.e. x = value - target
 			//
-			// Target velocity turns into an offset to the position
+			// The target velocity turns into an offset to the position
 			T AdjustedTarget = InTargetValue + InTargetValueRate * InSmoothingTime;
 
 			const float W = 2.0f / InSmoothingTime;
@@ -1635,8 +1635,6 @@ public:
 			// Approximation requires A < 1, so DeltaTime < SmoothingTime / 2
 			const float Exp = InvExpApprox(A);
 			T X0 = InOutValue - AdjustedTarget;
-			// Target velocity turns into an offset to the position
-			X0 -= InTargetValueRate * InSmoothingTime;		
 			const T B = (InOutValueRate + X0 * W) * InDeltaTime;
 			InOutValue = AdjustedTarget + (X0 + B) * Exp;
 			InOutValueRate = (InOutValueRate - B * W) * Exp;
