Add `UE_PLATFORM_REQUIRES_NONNULL_MEMOP_ARGS` and use it to workaround currently undefined behavior with C stdlib memory operations (will eventually be fixed by n3322).

#Synced-CL 51157306

diff --git a/Engine/Source/Runtime/Core/Public/GenericPlatform/GenericPlatformMemory.h b/Engine/Source/Runtime/Core/Public/GenericPlatform/GenericPlatformMemory.h
--- a/Engine/Source/Runtime/Core/Public/GenericPlatform/GenericPlatformMemory.h
+++ b/Engine/Source/Runtime/Core/Public/GenericPlatform/GenericPlatformMemory.h
@@ -570,44 +570,140 @@ struct FGenericPlatformMemory
 	 */
 	static UE_FORCEINLINE_HINT void* Memmove( void* Dest, const void* Src, SIZE_T Count )
 	{
+#if UE_PLATFORM_REQUIRES_NONNULL_MEMOP_ARGS
+		// On this platform `memmove` is required to not have `nullptr` as `Dest` or `Src`.
+		// 
+		// This will eventually go away when `n3322` is adopted, but for now we workaround
+		// this by checking for `Count == 0` and bailing early (the only defined behavior
+		// way going forward).
+		if (UNLIKELY(Count == 0))
+		{
+			return Dest;
+		}
+#endif
+
 		return memmove( Dest, Src, Count );
 	}
 
 	static UE_FORCEINLINE_HINT int32 Memcmp( const void* Buf1, const void* Buf2, SIZE_T Count )
 	{
+#if UE_PLATFORM_REQUIRES_NONNULL_MEMOP_ARGS
+		// On this platform `memcmp` is required to not have `nullptr` as `Dest` or `Src`.
+		// 
+		// This will eventually go away when `n3322` is adopted, but for now we workaround
+		// this by checking for `Count == 0` and bailing early (the only defined behavior
+		// way going forward).
+		if (UNLIKELY(Count == 0))
+		{
+			return 0;
+		}
+#endif
+
 		return memcmp( Buf1, Buf2, Count );
 	}
 
 	static UE_FORCEINLINE_HINT void* Memset(void* Dest, uint8 Char, SIZE_T Count)
 	{
+#if UE_PLATFORM_REQUIRES_NONNULL_MEMOP_ARGS
+		// On this platform `memset` is required to not have `nullptr` as `Dest`.
+		//
+		// This will eventually go away when `n3322` is adopted, but for now we workaround
+		// this by checking for `Count == 0` and bailing early (the only defined behavior
+		// way going forward).
+		if (UNLIKELY(Count == 0))
+		{
+			return Dest;
+		}
+#endif
+
 		return memset( Dest, Char, Count );
 	}
 
 	static UE_FORCEINLINE_HINT void* Memzero(void* Dest, SIZE_T Count)
 	{
+#if UE_PLATFORM_REQUIRES_NONNULL_MEMOP_ARGS
+		// On this platform `memset` is required to not have `nullptr` as `Dest`.
+		//
+		// This will eventually go away when `n3322` is adopted, but for now we workaround
+		// this by checking for `Count == 0` and bailing early (the only defined behavior
+		// way going forward).
+		if (UNLIKELY(Count == 0))
+		{
+			return Dest;
+		}
+#endif
+
 		return memset( Dest, 0, Count );
 	}
 
 	static UE_FORCEINLINE_HINT void* Memcpy(void* Dest, const void* Src, SIZE_T Count)
 	{
+#if UE_PLATFORM_REQUIRES_NONNULL_MEMOP_ARGS
+		// On this platform `memcpy` is required to not have `nullptr` as `Dest` or `Src`.
+		// 
+		// This will eventually go away when `n3322` is adopted, but for now we workaround
+		// this by checking for `Count == 0` and bailing early (the only defined behavior
+		// way going forward).
+		if (UNLIKELY(Count == 0))
+		{
+			return Dest;
+		}
+#endif
+
 		return memcpy( Dest, Src, Count );
 	}
 
 	/** Memcpy optimized for big blocks. */
 	static UE_FORCEINLINE_HINT void* BigBlockMemcpy(void* Dest, const void* Src, SIZE_T Count)
 	{
+#if UE_PLATFORM_REQUIRES_NONNULL_MEMOP_ARGS
+		// On this platform `memcpy` is required to not have `nullptr` as `Dest` or `Src`.
+		// 
+		// This will eventually go away when `n3322` is adopted, but for now we workaround
+		// this by checking for `Count == 0` and bailing early (the only defined behavior
+		// way going forward).
+		if (UNLIKELY(Count == 0))
+		{
+			return Dest;
+		}
+#endif
+
 		return memcpy(Dest, Src, Count);
 	}
 
 	/** On some platforms memcpy optimized for big blocks that avoid L2 cache pollution are available */
 	static UE_FORCEINLINE_HINT void* StreamingMemcpy(void* Dest, const void* Src, SIZE_T Count)
 	{
+#if UE_PLATFORM_REQUIRES_NONNULL_MEMOP_ARGS
+		// On this platform `memcpy` is required to not have `nullptr` as `Dest` or `Src`.
+		// 
+		// This will eventually go away when `n3322` is adopted, but for now we workaround
+		// this by checking for `Count == 0` and bailing early (the only defined behavior
+		// way going forward).
+		if (UNLIKELY(Count == 0))
+		{
+			return Dest;
+		}
+#endif
+
 		return memcpy( Dest, Src, Count );
 	}
 
 	/** On some platforms memcpy can be distributed over multiple threads for throughput. */
 	static inline void* ParallelMemcpy(void* Dest, const void* Src, SIZE_T Count, EMemcpyCachePolicy Policy = EMemcpyCachePolicy::StoreCached)
 	{
+#if UE_PLATFORM_REQUIRES_NONNULL_MEMOP_ARGS
+		// On this platform `memcpy` is required to not have `nullptr` as `Dest` or `Src`.
+		// 
+		// This will eventually go away when `n3322` is adopted, but for now we workaround
+		// this by checking for `Count == 0` and bailing early (the only defined behavior
+		// way going forward).
+		if (UNLIKELY(Count == 0))
+		{
+			return Dest;
+		}
+#endif
+
 		(void)Policy;
 		return memcpy( Dest, Src, Count );
 	}
diff --git a/Engine/Source/Runtime/Core/Public/HAL/Platform.h b/Engine/Source/Runtime/Core/Public/HAL/Platform.h
--- a/Engine/Source/Runtime/Core/Public/HAL/Platform.h
+++ b/Engine/Source/Runtime/Core/Public/HAL/Platform.h
@@ -1082,6 +1082,14 @@ int32 tchar_main(int32 ArgC, TCHAR* ArgV[])
 #	define UE_NO_PROFILE_SVE_BUG
 #endif
 
+// By default we assume that platforms require us to not have `nullptr` as arguments to
+// memory operations (like `memcpy`, `memset`, `memmove`, etc). `n3322` will define that
+// having `nullptr` as arguments to these functions with a size of 0 is ok, rather than
+// the UB it caused previously.
+#ifndef UE_PLATFORM_REQUIRES_NONNULL_MEMOP_ARGS
+#	define UE_PLATFORM_REQUIRES_NONNULL_MEMOP_ARGS 1
+#endif
+
 //--------------------------------------------------------------------------------------------
 // POD types refactor for porting old code:
 //
diff --git a/Engine/Source/Runtime/Core/Public/Microsoft/MSVCPlatform.h b/Engine/Source/Runtime/Core/Public/Microsoft/MSVCPlatform.h
--- a/Engine/Source/Runtime/Core/Public/Microsoft/MSVCPlatform.h
+++ b/Engine/Source/Runtime/Core/Public/Microsoft/MSVCPlatform.h
@@ -31,4 +31,7 @@
 	#endif
 #endif
 
-#define UE_NO_PROFILE_ATTRIBUTE
\ No newline at end of file
+#define UE_NO_PROFILE_ATTRIBUTE
+
+// MSVC doesn't complain if we pass `nullptr` to memory operations.
+#define UE_PLATFORM_REQUIRES_NONNULL_MEMOP_ARGS 0
