[Core] Optimize compact set reset to only clear specific hash elements instead of memsetting the whole hash table if only a small number are used

#tests Low level tests and ran editor+pie

#Synced-CL 51101133
#Synced-CL 51101133

diff --git a/Engine/Source/Runtime/Core/Public/Containers/CompactHashTable.h b/Engine/Source/Runtime/Core/Public/Containers/CompactHashTable.h
--- a/Engine/Source/Runtime/Core/Public/Containers/CompactHashTable.h
+++ b/Engine/Source/Runtime/Core/Public/Containers/CompactHashTable.h
@@ -238,6 +238,16 @@ namespace UE::Core::CompactHashTable
 
 		return IndexTypeToSizeType(NextIndex);
 	}
+
+	/*
+	 * Only remove the hash element, not the next index, this is only to be used when clearing the hash data to only hit the specific hash elements instead of writing the entirety of memory
+	 */
+	template<typename IndexType>
+	inline void RemoveHashElement(const uint32 Key, IndexType* HashData, const uint32 HashCount)
+	{
+		const uint32 HashIndex = Key & (HashCount - 1);
+		HashData[HashIndex] = (IndexType)INDEX_NONE;
+	}
 }
 
 /* Helper to lookup a type from a type size, todo: Add support for 3 byte lookup */
@@ -411,4 +421,11 @@ public:
 		#undef UE_COMPACTHASHTABLE_EXECUTEBYTYPE
 		return INDEX_NONE;
 	}
+
+	UE_FORCEINLINE_HINT void RemoveHashElement(uint32 Key) const
+	{
+		#define UE_COMPACTHASHTABLE_EXECUTEBYTYPE(Type) UE::Core::CompactHashTable::RemoveHashElement(Key, (Type*)(HashData), HashCount)
+		UE_COMPACTHASHTABLE_CALLBYTYPE(NextIndexCount);
+		#undef UE_COMPACTHASHTABLE_EXECUTEBYTYPE
+	}
 };
diff --git a/Engine/Source/Runtime/Core/Public/Containers/CompactSet.h.inl b/Engine/Source/Runtime/Core/Public/Containers/CompactSet.h.inl
--- a/Engine/Source/Runtime/Core/Public/Containers/CompactSet.h.inl
+++ b/Engine/Source/Runtime/Core/Public/Containers/CompactSet.h.inl
@@ -502,6 +502,7 @@ public:
 		this->NumElements = 0;
 
 		ResizeAllocation(ExpectedNumElements);
+
 		if (this->MaxElements > 0)
 		{
 			GetHashTableView().Reset();
@@ -513,10 +514,29 @@ public:
 	{
 		if (this->NumElements > 0)
 		{
+			FCompactHashTableView HashTable = GetHashTableView();
+
+			if (this->NumElements < (SizeType)(HashTable.GetHashCount() / UE_COMPACTSET_UNHASH_THRESHOLD))
+			{
+				// if we've only used a portion of the full allocated size then it can be faster
+				// to only clear out the used hash elements instead of memsetting the whole of the memory
+				// The exact number can be hard to determine since we have to recalculate the hash value to clear it
+
+				for (const ElementType& Element : *this)
+				{
+					const KeyInitType Key = KeyFuncs::GetSetKey(Element);
+					const uint32 KeyHash = KeyFuncs::GetKeyHash(Key);
+
+					HashTable.RemoveHashElement(KeyHash);
+				}
+			}
+			else
+			{
+				HashTable.Reset();
+			}
+
 			DestructItems(GetData(), this->NumElements);
 			this->NumElements = 0;
-			
-			GetHashTableView().Reset();
 		}
 	}
 
diff --git a/Engine/Source/Runtime/Core/Public/Containers/SetUtilities.h b/Engine/Source/Runtime/Core/Public/Containers/SetUtilities.h
--- a/Engine/Source/Runtime/Core/Public/Containers/SetUtilities.h
+++ b/Engine/Source/Runtime/Core/Public/Containers/SetUtilities.h
@@ -8,6 +8,11 @@
 #include "Templates/Tuple.h"
 #include "Templates/UnrealTypeTraits.h"
 
+// These divisors will determine at what point it's more optimal to individually remove elements from the hash table vs just using memset to reset the whole thing at once
+// Sparse set stores the hashtable index so it's fairly cheap to remove vs compact set that has to recalculate the hash key which is more expensive
+#define UE_SPARSESET_UNHASH_THRESHOLD 4
+#define UE_COMPACTSET_UNHASH_THRESHOLD 32
+
 /**
  * Traits class which determines whether or not a type is a TSet.
  */
diff --git a/Engine/Source/Runtime/Core/Public/Containers/SparseSet.h.inl b/Engine/Source/Runtime/Core/Public/Containers/SparseSet.h.inl
--- a/Engine/Source/Runtime/Core/Public/Containers/SparseSet.h.inl
+++ b/Engine/Source/Runtime/Core/Public/Containers/SparseSet.h.inl
@@ -1365,7 +1365,7 @@ private:
 		FSetElementId* HashPtr = Hash.GetAllocation();
 
 		// Check if it should be faster to clear the hash by going through elements instead of resetting the whole hash
-		if (Num() < (HashSize / 4))
+		if (Num() < (HashSize / UE_SPARSESET_UNHASH_THRESHOLD))
 		{
 			// Faster path: only reset hash buckets to invalid for elements in the hash
 			for (const SetElementType& Element: Elements)
diff --git a/Engine/Source/Runtime/Core/Tests/Containers/SetTests.cpp b/Engine/Source/Runtime/Core/Tests/Containers/SetTests.cpp
--- a/Engine/Source/Runtime/Core/Tests/Containers/SetTests.cpp
+++ b/Engine/Source/Runtime/Core/Tests/Containers/SetTests.cpp
@@ -279,6 +279,45 @@ struct TSetTestHelper
 			}
 		}
 
+		SECTION("Clear Tests")
+		{
+			SetType<int32> Set;
+
+			const uint32 PartialElements = 4;
+			// + 1 so it's immune to < vs <=, we want to make sure to test both unhash and memset scenarios for both set types
+			const uint32 MaxElements = (PartialElements + 1) * FMath::Max(UE_SPARSESET_UNHASH_THRESHOLD, UE_COMPACTSET_UNHASH_THRESHOLD);
+			Set.Empty(MaxElements);
+			CHECK(Set.Max() == MaxElements);
+
+			// There's not really much to test here, just run multiple passes of reusing the set to make sure things don't explode
+			for (uint32 LoopIndex = 0; LoopIndex < 4; ++LoopIndex)
+			{
+				const int32 Seed = MaxElements * (LoopIndex & 1);
+
+				for (uint32 PartialIndex = 0; PartialIndex < PartialElements; ++PartialIndex)
+				{
+					Set.Add(Seed + PartialIndex);
+					CHECK(Set.Contains(Seed + PartialIndex));
+				}
+
+				CHECK(Set.Num() == PartialElements);
+				Set.Reset();
+				CHECK(Set.Num() == 0);
+				CHECK(Set.Max() == MaxElements);
+
+				for (uint32 FullIndex = 0; FullIndex < MaxElements; ++FullIndex)
+				{
+					Set.Add(Seed + FullIndex);
+					CHECK(Set.Contains(Seed + FullIndex));
+				}
+
+				CHECK(Set.Num() == MaxElements);
+				Set.Reset();
+				CHECK(Set.Num() == 0);
+				CHECK(Set.Max() == MaxElements);
+			}
+		}
+
 		SECTION("Iteration")
 		{
 			int32 IntTestItems[256];
