[Core] Add compile options to change checkslows to regular checks and to clear out next indices on removal.

#tests Low level tests and ran editor+pie

#Synced-CL 51169782

diff --git a/Engine/Source/Programs/UnrealBuildTool/Configuration/Rules/ReadOnlyTargetRules.cs b/Engine/Source/Programs/UnrealBuildTool/Configuration/Rules/ReadOnlyTargetRules.cs
--- a/Engine/Source/Programs/UnrealBuildTool/Configuration/Rules/ReadOnlyTargetRules.cs
+++ b/Engine/Source/Programs/UnrealBuildTool/Configuration/Rules/ReadOnlyTargetRules.cs
@@ -284,6 +284,8 @@ namespace UnrealBuildTool
 
 		public bool bUseCompactSetAsDefault => Inner.bUseCompactSetAsDefault;
 
+		public bool bEnableCompactSetDebug => Inner.bEnableCompactSetDebug;
+
 		public bool bUseAndroidJni => Inner.bUseAndroidJni;
 
 		public FPSemanticsMode FPSemantics => Inner.FPSemantics;
diff --git a/Engine/Source/Programs/UnrealBuildTool/Configuration/Rules/TargetRules.cs b/Engine/Source/Programs/UnrealBuildTool/Configuration/Rules/TargetRules.cs
--- a/Engine/Source/Programs/UnrealBuildTool/Configuration/Rules/TargetRules.cs
+++ b/Engine/Source/Programs/UnrealBuildTool/Configuration/Rules/TargetRules.cs
@@ -1707,6 +1707,14 @@ namespace UnrealBuildTool
 		[XmlConfigFile(Category = "BuildConfiguration")]
 		public bool bUseCompactSetAsDefault { get; set; } = false;
 
+		/// <summary>
+		/// Enable additional debug features for compact sets to help detect improper use better, these features tend to have undesirable perf impacts
+		/// </summary>
+		[RequiresUniqueBuildEnvironment]
+		[CommandLine("-EnableCompactSetDebug")]
+		[XmlConfigFile(Category = "BuildConfiguration")]
+		public bool bEnableCompactSetDebug { get; set; } = false;
+
 		/// <summary>
 		/// Whether to use Android Java Native Interface.
 		/// </summary>
diff --git a/Engine/Source/Programs/UnrealBuildTool/Configuration/UEBuildTarget.cs b/Engine/Source/Programs/UnrealBuildTool/Configuration/UEBuildTarget.cs
--- a/Engine/Source/Programs/UnrealBuildTool/Configuration/UEBuildTarget.cs
+++ b/Engine/Source/Programs/UnrealBuildTool/Configuration/UEBuildTarget.cs
@@ -6632,6 +6632,15 @@ namespace UnrealBuildTool
 				GlobalCompileEnvironment.Definitions.Add("UE_USE_COMPACT_SET_AS_DEFAULT=0");
 			}
 
+			if (Rules.bEnableCompactSetDebug)
+			{
+				GlobalCompileEnvironment.Definitions.Add("UE_ENABLE_COMPACT_SET_DEBUG=1");
+			}
+			else
+			{
+				GlobalCompileEnvironment.Definitions.Add("UE_ENABLE_COMPACT_SET_DEBUG=0");
+			}
+
 			// Define the custom config if specified, this should not be set in an editor build
 			if (!bUseSharedBuildEnvironment && !String.IsNullOrEmpty(Rules.CustomConfig))
 			{
diff --git a/Engine/Source/Runtime/Core/Public/Containers/CompactHashTable.h b/Engine/Source/Runtime/Core/Public/Containers/CompactHashTable.h
--- a/Engine/Source/Runtime/Core/Public/Containers/CompactHashTable.h
+++ b/Engine/Source/Runtime/Core/Public/Containers/CompactHashTable.h
@@ -4,6 +4,20 @@
 
 #include "Containers/ArrayView.h"
 
+// Enabling UE_ENABLE_COMPACT_SET_DEBUG can add additional features to help track down issues with improper use of compact hash tables
+#if !defined(UE_ENABLE_COMPACT_SET_DEBUG)
+#define UE_ENABLE_COMPACT_SET_DEBUG 0
+#endif
+
+#if UE_ENABLE_COMPACT_SET_DEBUG
+#define UE_COMPACTHASH_CHECK check
+#define UE_COMPACTHASH_CHECKF checkf
+#else
+#define UE_COMPACTHASH_CHECK checkSlow
+#define UE_COMPACTHASH_CHECKF checkfSlow
+#endif
+
+
 /* Compact hash table has two distinct features to keep it small
 	1. The Index type adapts to the number of elements, so a smaller table will use a smaller type
 	2. Holes in the table are patched up so all slots up to count are guaranteed to be valid
@@ -117,7 +131,7 @@ namespace UE::Core::CompactHashTable
 	template<typename IndexType>
 	[[nodiscard]] inline uint32 GetFirstByIndex(uint32 HashIndex, const IndexType* HashData, const uint32 HashCount)
 	{
-		checkSlow(HashIndex < HashCount);
+		UE_COMPACTHASH_CHECK(HashIndex < HashCount);
 		return IndexTypeToSizeType(HashData[HashIndex]);
 	}
 
@@ -125,7 +139,7 @@ namespace UE::Core::CompactHashTable
 	template<typename IndexType>
 	[[nodiscard]] inline uint32 GetNext(uint32 Index, const IndexType* NextIndexData, const uint32 NextIndexCount)
 	{
-		checkSlow(Index < NextIndexCount);
+		UE_COMPACTHASH_CHECK(Index < NextIndexCount);
 		return IndexTypeToSizeType(NextIndexData[Index]);
 	}
 
@@ -140,7 +154,7 @@ namespace UE::Core::CompactHashTable
 				// Return the first match, regardless of whether the set has multiple matches for the key or not.
 				return ElementIndex;
 			}
-			checkSlow(ElementIndex < NextIndexCount);
+			UE_COMPACTHASH_CHECK(ElementIndex < NextIndexCount);
 		}
 		return INDEX_NONE;
 	}
@@ -149,7 +163,7 @@ namespace UE::Core::CompactHashTable
 	template<typename IndexType>
 	inline void Add(uint32 Index, uint32 Key, IndexType* HashData, const uint32 HashCount, IndexType* NextIndexData, const uint32 NextIndexCount)
 	{
-		checkSlow(Index < NextIndexCount);
+		UE_COMPACTHASH_CHECK(Index < NextIndexCount);
 
 		const uint32 HashIndex = Key & (HashCount - 1);
 		NextIndexData[Index] = HashData[HashIndex];
@@ -162,7 +176,7 @@ namespace UE::Core::CompactHashTable
 	template<typename IndexType>
 	inline uint32 Remove(const uint32 Index, const uint32 Key, const uint32 LastIndex, uint32 OptLastKey, IndexType* HashData, const uint32 HashCount, IndexType* NextIndexData, const uint32 NextIndexCount)
 	{
-		checkSlow(LastIndex < NextIndexCount && Index <= LastIndex);
+		UE_COMPACTHASH_CHECK(LastIndex < NextIndexCount && Index <= LastIndex);
 
 		IndexType NextIndex = UE::Core::CompactHashTable::Private::RemoveInternal<IndexType>(Index, Key, HashData, HashCount, NextIndexData, NextIndexCount);
 
@@ -184,6 +198,12 @@ namespace UE::Core::CompactHashTable
 			}
 		}
 
+#if UE_BUILD_SHIPPING
+		// This shouldd't be needed if everything is working correctly, only doing this in shipping builds to prevent accidental
+		// infinite loops, but leaving this off in non-shipping builds so we can properly detect and address the root cause
+		NextIndexData[LastIndex] = (IndexType)INDEX_NONE;
+#endif
+
 		// If the next index points to the last element, then it will get relocated into the current element so return the current instead of the next
 		return NextIndex == LastIndex ? Index : IndexTypeToSizeType(NextIndex);
 	}
@@ -194,7 +214,7 @@ namespace UE::Core::CompactHashTable
 	template<typename IndexType>
 	inline uint32 RemoveStable(const uint32 Index, const uint32 Key, IndexType* HashData, const uint32 HashCount, IndexType* NextIndexData, const uint32 NextIndexCount)
 	{
-		checkSlow(Index < NextIndexCount);
+		UE_COMPACTHASH_CHECK(Index < NextIndexCount);
 
 		constexpr IndexType InvalidIndex = (IndexType)INDEX_NONE;
 
@@ -287,20 +307,20 @@ public:
 
 	[[nodiscard]] inline uint32 GetNext(uint32 Index, uint32 CurrentCount) const
 	{
-		checkSlow(CurrentCount <= ElementCount);
+		UE_COMPACTHASH_CHECK(CurrentCount <= ElementCount);
 		return UE::Core::CompactHashTable::GetNext(Index, NextIndexData, CurrentCount);
 	}
 
 	template<typename PredicateType>
 	[[nodiscard]] inline uint32 Find(uint32 Key, uint32 CurrentCount, const PredicateType& Predicate) const
 	{
-		checkSlow(CurrentCount <= ElementCount);
+		UE_COMPACTHASH_CHECK(CurrentCount <= ElementCount);
 		return UE::Core::CompactHashTable::Find(Key, HashData, HashCount, NextIndexData, CurrentCount, Predicate);
 	}
 	
 	inline void Add(uint32 CurrentCount, uint32 Key)
 	{
-		checkSlow(CurrentCount < ElementCount);
+		UE_COMPACTHASH_CHECK(CurrentCount < ElementCount);
 		UE::Core::CompactHashTable::Add(CurrentCount, Key, HashData, HashCount, NextIndexData, CurrentCount + 1);
 	}
 
@@ -326,9 +346,9 @@ public:
 	, NextIndexCount(InNextIndexCount)
 	, HashCount(InHashCount)
 	{
-		checkSlow(Memory != nullptr && InNextIndexCount > 0 && InHashCount > 0 && MemorySize > 0);
-		checkfSlow(MemorySize == UE::Core::CompactHashTable::GetMemoryRequiredInBytes(NextIndexCount, HashCount), TEXT("Got %d bytes, expected %d bytes"), MemorySize, UE::Core::CompactHashTable::GetMemoryRequiredInBytes(NextIndexCount, HashCount));
-		checkSlow(FMath::IsPowerOfTwo(HashCount));
+		UE_COMPACTHASH_CHECK(Memory != nullptr && InNextIndexCount > 0 && InHashCount > 0 && MemorySize > 0);
+		UE_COMPACTHASH_CHECKF(MemorySize == UE::Core::CompactHashTable::GetMemoryRequiredInBytes(NextIndexCount, HashCount), TEXT("Got %lld bytes, expected %lld bytes"), MemorySize, UE::Core::CompactHashTable::GetMemoryRequiredInBytes(NextIndexCount, HashCount));
+		UE_COMPACTHASH_CHECK(FMath::IsPowerOfTwo(HashCount));
 	}
 
 	[[nodiscard]] uint32 GetHashCount() const
@@ -356,7 +376,7 @@ public:
 
 	[[nodiscard]] inline uint32 GetNext(uint32 Index, uint32 CurrentCount) const
 	{
-		checkSlow(CurrentCount <= NextIndexCount);
+		UE_COMPACTHASH_CHECK(CurrentCount <= NextIndexCount);
 	
 		#define UE_COMPACTHASHTABLE_EXECUTEBYTYPE(Type) return UE::Core::CompactHashTable::GetNext(Index, (const Type *)NextIndexData, CurrentCount)
 		UE_COMPACTHASHTABLE_CALLBYTYPE(NextIndexCount);
@@ -367,7 +387,7 @@ public:
 	template<typename PredicateType>
 	[[nodiscard]] inline uint32 Find(uint32 Key, uint32 CurrentCount, const PredicateType& Predicate) const
 	{
-		checkSlow(CurrentCount <= NextIndexCount);
+		UE_COMPACTHASH_CHECK(CurrentCount <= NextIndexCount);
 		#define UE_COMPACTHASHTABLE_EXECUTEBYTYPE(Type) return UE::Core::CompactHashTable::Find(Key, (const Type *)HashData, HashCount, (const Type *)NextIndexData, CurrentCount, Predicate)
 		UE_COMPACTHASHTABLE_CALLBYTYPE(NextIndexCount);
 		#undef UE_COMPACTHASHTABLE_EXECUTEBYTYPE
@@ -400,7 +420,7 @@ public:
 	
 	inline void Add(uint32 Index, uint32 Key) const
 	{
-		checkSlow(Index < NextIndexCount);
+		UE_COMPACTHASH_CHECK(Index < NextIndexCount);
 		#define UE_COMPACTHASHTABLE_EXECUTEBYTYPE(Type) UE::Core::CompactHashTable::Add(Index, Key, (Type*)(HashData), HashCount, (Type*)(NextIndexData), NextIndexCount)
 		UE_COMPACTHASHTABLE_CALLBYTYPE(NextIndexCount);
 		#undef UE_COMPACTHASHTABLE_EXECUTEBYTYPE
