- Improved sentinel file handling

#Synced-CL 50913554

diff --git a/Engine/Source/Runtime/Core/Private/GenericPlatform/GenericPlatformProcess.cpp b/Engine/Source/Runtime/Core/Private/GenericPlatform/GenericPlatformProcess.cpp
--- a/Engine/Source/Runtime/Core/Private/GenericPlatform/GenericPlatformProcess.cpp
+++ b/Engine/Source/Runtime/Core/Private/GenericPlatform/GenericPlatformProcess.cpp
@@ -881,6 +881,10 @@ TSharedPtr<IProcessSentinel> FGenericPlatformProcess::CreateProcessSentinel(cons
 {
 	TSharedPtr<IProcessSentinel> Sentinel = FPlatformProcess::CreateProcessSentinelObject();
 	
-	Sentinel->CreateSentinelFile(Name, Contents);
+	// if we couldn't create a file, then return null to indicate it's not usable (likely another process already made it)
+	if (Sentinel->CreateSentinelFile(Name, Contents) == false)
+	{
+		return nullptr;
+	}
 	return Sentinel;
 }
diff --git a/Engine/Source/Runtime/Core/Private/Windows/WindowsPlatformProcess.cpp b/Engine/Source/Runtime/Core/Private/Windows/WindowsPlatformProcess.cpp
--- a/Engine/Source/Runtime/Core/Private/Windows/WindowsPlatformProcess.cpp
+++ b/Engine/Source/Runtime/Core/Private/Windows/WindowsPlatformProcess.cpp
@@ -672,22 +672,22 @@ FProcHandle FWindowsPlatformProcess::CreateProcInternal(const TCHAR* URL, const
 		},
 		AttributeList
 	};
-	PROCESS_INFORMATION ProcessInfomation;
+	PROCESS_INFORMATION ProcessInformation;
 
-	if (!CreateProcess(nullptr, CommandLine.GetCharArray().GetData(), nullptr, nullptr, bInheritHandles || AttributeList, CreationFlags, nullptr, OptionalWorkingDirectory, &StartupInfoEx.StartupInfo, &ProcessInfomation))
+	if (!CreateProcess(nullptr, CommandLine.GetCharArray().GetData(), nullptr, nullptr, bInheritHandles || AttributeList, CreationFlags, nullptr, OptionalWorkingDirectory, &StartupInfoEx.StartupInfo, &ProcessInformation))
 	{
 		HandleError();
 		return {};
 	}
 
-	CloseHandle(ProcessInfomation.hThread);
+	CloseHandle(ProcessInformation.hThread);
 
 	if (OutProcessID)
 	{
-		*OutProcessID = ProcessInfomation.dwProcessId;
+		*OutProcessID = ProcessInformation.dwProcessId;
 	}
 
-	return FProcHandle{ProcessInfomation.hProcess};
+	return FProcHandle{ProcessInformation.hProcess};
 }
 #include "Windows/HideWindowsPlatformTypes.h"
 
@@ -2776,7 +2776,10 @@ public:
 	virtual ~FWindowsProcessSentinel()
 	{
 		// this will also delete the file
-		CloseHandle(FileHandle);
+		if (FileHandle != INVALID_HANDLE_VALUE)
+		{
+			CloseHandle(FileHandle);
+		}
 	}
 
 	virtual bool CreateSentinelFile(const TCHAR* Name, const FString& Contents) override
@@ -2844,7 +2847,7 @@ public:
 		return true;
 	}
 	
-	HANDLE FileHandle;
+	HANDLE FileHandle = INVALID_HANDLE_VALUE;
 };
 
 TSharedPtr<IProcessSentinel> FWindowsPlatformProcess::CreateProcessSentinelObject()
