Logging: Disable the backlog at the end of PreInit for program targets

The backlog grows unbounded for program targets if they never call GEngineLoop.Init(), and most do not. A program that needs the backlog to persist through GEngineLoop.Init() can enable it explicitly before GEngineLoop.PreInit() by calling GLog->EnableBacklog(true).

diff --git a/Engine/Source/Runtime/Core/Private/Misc/OutputDeviceRedirector.cpp b/Engine/Source/Runtime/Core/Private/Misc/OutputDeviceRedirector.cpp
--- a/Engine/Source/Runtime/Core/Private/Misc/OutputDeviceRedirector.cpp
+++ b/Engine/Source/Runtime/Core/Private/Misc/OutputDeviceRedirector.cpp
@@ -255,7 +255,7 @@ struct FOutputDeviceRedirectorState
 		}
 	}
 
-	FORCENOINLINE void EnableBacklog(bool bEnable);
+	FORCENOINLINE bool EnableBacklog(bool bEnable);
 	FORCENOINLINE void AddToBacklog(const FLogRecord& Record);
 	FORCENOINLINE void AddToBacklog(const TCHAR* Data, ELogVerbosity::Type Verbosity, const FName& Category, double Time);
 	FORCENOINLINE void SerializeBacklog(TNotNull<FOutputDevice*> OutputDevice);
@@ -650,7 +650,7 @@ void FOutputDeviceRedirectorState::WaitFence(uint64 Index)
 }
 
 UE_AUTORTFM_NOAUTORTFM
-void FOutputDeviceRedirectorState::EnableBacklog(bool bEnable)
+bool FOutputDeviceRedirectorState::EnableBacklog(bool bEnable)
 {
 	FWriteScopeLock ScopeLock(BacklogLock);
 	const bool bWasEnabled = bEnableBacklog;
@@ -688,6 +688,7 @@ void FOutputDeviceRedirectorState::EnableBacklog(bool bEnable)
 			FlushBufferedItemsTo(FVisitor{*this, IsPanicThread(ThreadId)});
 		}
 	}
+	return bWasEnabled;
 }
 
 UE_AUTORTFM_NOAUTORTFM
@@ -801,9 +802,9 @@ void FOutputDeviceRedirector::SerializeBacklog(FOutputDevice* OutputDevice)
 }
 
 UE_AUTORTFM_NOAUTORTFM // Not currently supported in a closed transaction
-void FOutputDeviceRedirector::EnableBacklog(bool bEnable)
+bool FOutputDeviceRedirector::EnableBacklog(bool bEnable)
 {
-	State->EnableBacklog(bEnable);
+	return State->EnableBacklog(bEnable);
 }
 
 UE_AUTORTFM_NOAUTORTFM // Not currently supported in a closed transaction
diff --git a/Engine/Source/Runtime/Core/Public/Misc/OutputDeviceRedirector.h b/Engine/Source/Runtime/Core/Public/Misc/OutputDeviceRedirector.h
--- a/Engine/Source/Runtime/Core/Public/Misc/OutputDeviceRedirector.h
+++ b/Engine/Source/Runtime/Core/Public/Misc/OutputDeviceRedirector.h
@@ -98,8 +98,9 @@ public:
 	 * Enables or disables the backlog.
 	 *
 	 * @param bEnable   Starts saving a backlog if true, disables and discards any backlog if false.
+	 * @return true if the backlog was previously enabled or false if it was not.
 	 */
-	CORE_API void EnableBacklog(bool bEnable);
+	CORE_API bool EnableBacklog(bool bEnable);
 
 	/**
 	 * Sets the current thread to be the thread that redirects logs to buffered output devices.
diff --git a/Engine/Source/Runtime/Launch/Private/LaunchEngineLoop.cpp b/Engine/Source/Runtime/Launch/Private/LaunchEngineLoop.cpp
--- a/Engine/Source/Runtime/Launch/Private/LaunchEngineLoop.cpp
+++ b/Engine/Source/Runtime/Launch/Private/LaunchEngineLoop.cpp
@@ -1744,9 +1744,19 @@ int32 FEngineLoop::PreInitPreStartupScreen(const TCHAR* CmdLine)
 	// until initialization is far enough along to try to start a dedicated primary thread.
 	GLog->SetCurrentThreadAsPrimaryThread();
 
-	// Most targets have the backlog enabled by default. Enable it as early as possible here for
+#if IS_PROGRAM
+	// Non-program targets have the backlog enabled by default. Enable it as early as possible here for
 	// cases where it is disabled by default. It will be disabled when module loading is complete.
-	GLog->EnableBacklog(true);
+	// Do not disable the backlog if the program enabled it explicitly before calling PreInit().
+	const bool bDisableBacklog = !GLog->EnableBacklog(true);
+	ON_SCOPE_EXIT
+	{
+		if (bDisableBacklog)
+		{
+			GLog->EnableBacklog(false);
+		}
+	};
+#endif
 
 	// Command line option for enabling named events
 	if (FParse::Param(CmdLine, TEXT("statnamedevents")))
