[Backout] - CL51228771
#rnx
Original CL Desc
-----------------------------------------------------------------
- Improved sentinel file handling

diff --git a/Engine/Source/Runtime/Core/Private/GenericPlatform/GenericPlatformProcess.cpp b/Engine/Source/Runtime/Core/Private/GenericPlatform/GenericPlatformProcess.cpp
--- a/Engine/Source/Runtime/Core/Private/GenericPlatform/GenericPlatformProcess.cpp
+++ b/Engine/Source/Runtime/Core/Private/GenericPlatform/GenericPlatformProcess.cpp
@@ -881,10 +881,6 @@ TSharedPtr<IProcessSentinel> FGenericPlatformProcess::CreateProcessSentinel(cons
 {
 	TSharedPtr<IProcessSentinel> Sentinel = FPlatformProcess::CreateProcessSentinelObject();
 	
-	// if we couldn't create a file, then return null to indicate it's not usable (likely another process already made it)
-	if (Sentinel->CreateSentinelFile(Name, Contents) == false)
-	{
-		return nullptr;
-	}
+	Sentinel->CreateSentinelFile(Name, Contents);
 	return Sentinel;
 }
diff --git a/Engine/Source/Runtime/Core/Private/Windows/WindowsPlatformProcess.cpp b/Engine/Source/Runtime/Core/Private/Windows/WindowsPlatformProcess.cpp
--- a/Engine/Source/Runtime/Core/Private/Windows/WindowsPlatformProcess.cpp
+++ b/Engine/Source/Runtime/Core/Private/Windows/WindowsPlatformProcess.cpp
@@ -672,22 +672,22 @@ FProcHandle FWindowsPlatformProcess::CreateProcInternal(const TCHAR* URL, const
 		},
 		AttributeList
 	};
-	PROCESS_INFORMATION ProcessInformation;
+	PROCESS_INFORMATION ProcessInfomation;
 
-	if (!CreateProcess(nullptr, CommandLine.GetCharArray().GetData(), nullptr, nullptr, bInheritHandles || AttributeList, CreationFlags, nullptr, OptionalWorkingDirectory, &StartupInfoEx.StartupInfo, &ProcessInformation))
+	if (!CreateProcess(nullptr, CommandLine.GetCharArray().GetData(), nullptr, nullptr, bInheritHandles || AttributeList, CreationFlags, nullptr, OptionalWorkingDirectory, &StartupInfoEx.StartupInfo, &ProcessInfomation))
 	{
 		HandleError();
 		return {};
 	}
 
-	CloseHandle(ProcessInformation.hThread);
+	CloseHandle(ProcessInfomation.hThread);
 
 	if (OutProcessID)
 	{
-		*OutProcessID = ProcessInformation.dwProcessId;
+		*OutProcessID = ProcessInfomation.dwProcessId;
 	}
 
-	return FProcHandle{ProcessInformation.hProcess};
+	return FProcHandle{ProcessInfomation.hProcess};
 }
 #include "Windows/HideWindowsPlatformTypes.h"
 
@@ -2776,10 +2776,7 @@ public:
 	virtual ~FWindowsProcessSentinel()
 	{
 		// this will also delete the file
-		if (FileHandle != INVALID_HANDLE_VALUE)
-		{
-			CloseHandle(FileHandle);
-		}
+		CloseHandle(FileHandle);
 	}
 
 	virtual bool CreateSentinelFile(const TCHAR* Name, const FString& Contents) override
@@ -2847,7 +2844,7 @@ public:
 		return true;
 	}
 	
-	HANDLE FileHandle = INVALID_HANDLE_VALUE;
+	HANDLE FileHandle;
 };
 
 TSharedPtr<IProcessSentinel> FWindowsPlatformProcess::CreateProcessSentinelObject()
