Logging: Removed caching of log templates

Construction of log templates has been significantly optimized and the caching no longer offers any performance benefit, while incurring a cost in program size and runtime memory use.

#Synced-CL 50967057
#Synced-CL 50967057

diff --git a/Engine/Source/Runtime/Core/Private/Logging/StructuredLog.cpp b/Engine/Source/Runtime/Core/Private/Logging/StructuredLog.cpp
--- a/Engine/Source/Runtime/Core/Private/Logging/StructuredLog.cpp
+++ b/Engine/Source/Runtime/Core/Private/Logging/StructuredLog.cpp
@@ -1259,83 +1259,6 @@ FORCENOINLINE static void BasicLogToTrace(const void* LogPoint, const TCHAR* For
 
 ///////////////////////////////////////////////////////////////////////////////////////////////////////////////////////
 
-class FStaticLogTemplateAllocator
-{
-public:
-	~FStaticLogTemplateAllocator()
-	{
-		for (FBlock* Block = Tail; Block;)
-		{
-			FBlock* Previous = Block->Previous;
-			FMemory::Free(Block);
-			Block = Previous;
-		}
-	}
-
-	void* Allocate(int32 Size)
-	{
-		TUniqueLock Lock(Mutex);
-		if (!Tail || TailOffset + Size > Tail->Size)
-		{
-			int32 NewSize = (sizeof(FBlock) + Size + BlockSize - 1) & ~(BlockSize - 1);
-			Tail = new(FMemory::Malloc(NewSize, alignof(FBlock))) FBlock{Tail, NewSize};
-			TailOffset = sizeof(FBlock);
-		}
-		void* Address = (uint8*)Tail + TailOffset;
-		// TODO: Aligned to 8 until unaligned pointer access has been tested on every platform.
-		TailOffset += (Size + 7) & ~7;
-		return Address;
-	}
-
-private:
-	inline constexpr static int32 BlockSize = 4096;
-
-	struct FBlock
-	{
-		FBlock* Previous = nullptr;
-		int32 Size = 0;
-	};
-
-	FBlock* Tail = nullptr;
-	int32 TailOffset = 0;
-	FMutex Mutex;
-};
-
-static FStaticLogTemplateAllocator& GetStaticLogTemplateAllocator()
-{
-	static FStaticLogTemplateAllocator Allocator;
-	return Allocator;
-}
-
-class FStaticLogTemplateStorage
-{
-public:
-	FStaticLogTemplateStorage() = default;
-	FStaticLogTemplateStorage(const FStaticLogTemplateStorage&) = delete;
-	FStaticLogTemplateStorage& operator=(const FStaticLogTemplateStorage&) = delete;
-
-	inline void* Allocate(int32 Size)
-	{
-		Data = GetStaticLogTemplateAllocator().Allocate(Size);
-		return Data;
-	}
-
-	inline const void* Get() const
-	{
-		return Data;
-	}
-
-	inline void* Detach()
-	{
-		ON_SCOPE_EXIT { Data = nullptr; };
-		return Data;
-	}
-
-private:
-	// Intentionally leaked because the allocator will free the memory on exit.
-	void* Data = nullptr;
-};
-
 // Serializing log fields to compact binary happens in its own function because that allows stack
 // space for the writer to be returned before calling into the output devices.
 FORCENOINLINE static FCbObject SerializeLogFields(
@@ -1423,48 +1346,14 @@ FORCENOINLINE static FCbObject SerializeLogFields(
 	return Writer.Save().AsObject();
 }
 
-template <typename StaticLogRecordType>
-FORCENOINLINE static FLogTemplate& CreateLogTemplate(const FLogCategoryBase& Category, const StaticLogRecordType& Log, const FLogField* Fields, const int32 FieldCount)
+inline static FInlineLogTemplate CreateLogTemplate(const FStaticLogRecord& Log, const FLogField* Fields, const int32 FieldCount)
 {
-	for (FLogTemplate* Template = Log.DynamicData.Template.load(std::memory_order_acquire);;)
-	{
-		if (Template && Template->GetUtf8Format() == Log.Format)
-		{
-			return *Template;
-		}
-
-		struct FTemplateCreator
-		{
-			static TLogTemplate<FStaticLogTemplateStorage> Create(const FStaticLogRecord& Log, const FLogField* Fields, int32 FieldCount)
-			{
-				return TLogTemplate<FStaticLogTemplateStorage>(Log.Format, {}, Fields, FieldCount);
-			}
-
-			static TLogTemplate<FStaticLogTemplateStorage> Create(const FStaticLocalizedLogRecord& Log, const FLogField* Fields, int32 FieldCount)
-			{
-				return TLogTemplate<FStaticLogTemplateStorage>(Log.TextNamespace, Log.TextKey, Log.Format, {}, Fields, FieldCount);
-			}
-		};
-
-		TLogTemplate<FStaticLogTemplateStorage> LocalTemplate = FTemplateCreator::Create(Log, Fields, FieldCount);
-		FLogTemplate* NewTemplate = LocalTemplate.Get();
-		if (LIKELY(Log.DynamicData.Template.compare_exchange_strong(Template, NewTemplate, std::memory_order_release, std::memory_order_acquire)))
-		{
-			LocalTemplate.Detach();
-			return *NewTemplate;
-		}
-	}
+	return FInlineLogTemplate(Log.Format, {}, Fields, FieldCount);
 }
 
-template <typename StaticLogRecordType>
-inline static FLogTemplate& EnsureLogTemplate(const FLogCategoryBase& Category, const StaticLogRecordType& Log, const FLogField* Fields, const int32 FieldCount)
+inline static FInlineLogTemplate CreateLogTemplate(const FStaticLocalizedLogRecord& Log, const FLogField* Fields, const int32 FieldCount)
 {
-	// Format can change on a static log record due to Live Coding.
-	if (FLogTemplate* Template = Log.DynamicData.Template.load(std::memory_order_acquire); LIKELY(Template && Template->GetUtf8Format() == Log.Format))
-	{
-		return *Template;
-	}
-	return CreateLogTemplate(Category, Log, Fields, FieldCount);
+	return FInlineLogTemplate(Log.TextNamespace, Log.TextKey, Log.Format, {}, Fields, FieldCount);
 }
 
 template <typename StaticLogRecordType>
@@ -1477,12 +1366,11 @@ inline static FLogRecord CreateLogRecord(const FLogCategoryBase& Category, const
 	}
 #endif
 
-	FLogTemplate& Template = EnsureLogTemplate(Category, Log, Fields, FieldCount);
+	const FInlineLogTemplate Template = CreateLogTemplate(Log, Fields, FieldCount);
 
 	FLogRecord Record;
 	Record.SetUtf8Format(Log.Format);
-	Record.SetTemplate(&Template);
-	Record.SetFields(SerializeLogFields(Template, Fields, FieldCount));
+	Record.SetFields(SerializeLogFields(*Template.Get(), Fields, FieldCount));
 	Record.SetFile(Log.File);
 	Record.SetLine(Log.Line);
 	Record.SetCategory(Category.GetCategoryName());
diff --git a/Engine/Source/Runtime/Core/Public/Logging/StructuredLog.h b/Engine/Source/Runtime/Core/Public/Logging/StructuredLog.h
--- a/Engine/Source/Runtime/Core/Public/Logging/StructuredLog.h
+++ b/Engine/Source/Runtime/Core/Public/Logging/StructuredLog.h
@@ -334,7 +334,6 @@ namespace UE::Logging::Private
 /** Data about a static log that is created on-demand. */
 struct FStaticLogDynamicData
 {
-	std::atomic<FLogTemplate*> Template = nullptr;
 #if LOGTRACE_ENABLED
 	std::atomic<bool> bInitializedTrace = false;
 #endif
@@ -347,7 +346,9 @@ struct FStaticLogRecord
 	UE_IF(UE_LOG_INCLUDE_SOURCE_LOCATION,, inline static constexpr) const ANSICHAR* File = nullptr;
 	UE_IF(UE_LOG_INCLUDE_SOURCE_LOCATION,, inline static constexpr) int32 Line = 0;
 	ELogVerbosity::Type Verbosity = ELogVerbosity::Log;
+#if LOGTRACE_ENABLED
 	FStaticLogDynamicData& DynamicData;
+#endif
 
 	// Workaround for https://developercommunity.visualstudio.com/t/Incorrect-warning-C4700-with-unrelated-s/10285950
 	constexpr FStaticLogRecord(
@@ -355,14 +356,16 @@ struct FStaticLogRecord
 		const ANSICHAR* InFile,
 		int32 InLine,
 		ELogVerbosity::Type InVerbosity,
-		FStaticLogDynamicData& InDynamicData)
+		FStaticLogDynamicData* InDynamicData)
 		: Format(InFormat)
 	#if UE_LOG_INCLUDE_SOURCE_LOCATION
 		, File(InFile)
 		, Line(InLine)
 	#endif
 		, Verbosity(InVerbosity)
-		, DynamicData(InDynamicData)
+	#if LOGTRACE_ENABLED
+		, DynamicData(*InDynamicData)
+	#endif
 	{
 	}
 };
@@ -376,7 +379,9 @@ struct FStaticLocalizedLogRecord
 	UE_IF(UE_LOG_INCLUDE_SOURCE_LOCATION,, inline static constexpr) const ANSICHAR* File = nullptr;
 	UE_IF(UE_LOG_INCLUDE_SOURCE_LOCATION,, inline static constexpr) int32 Line = 0;
 	ELogVerbosity::Type Verbosity = ELogVerbosity::Log;
+#if LOGTRACE_ENABLED
 	FStaticLogDynamicData& DynamicData;
+#endif
 
 	// Workaround for https://developercommunity.visualstudio.com/t/Incorrect-warning-C4700-with-unrelated-s/10285950
 	constexpr FStaticLocalizedLogRecord(
@@ -386,7 +391,7 @@ struct FStaticLocalizedLogRecord
 		const ANSICHAR* InFile,
 		int32 InLine,
 		ELogVerbosity::Type InVerbosity,
-		FStaticLogDynamicData& InDynamicData)
+		FStaticLogDynamicData* InDynamicData)
 		: TextNamespace(InTextNamespace)
 		, TextKey(InTextKey)
 		, Format(InFormat)
@@ -395,7 +400,9 @@ struct FStaticLocalizedLogRecord
 		, Line(InLine)
 	#endif
 		, Verbosity(InVerbosity)
-		, DynamicData(InDynamicData)
+	#if LOGTRACE_ENABLED
+		, DynamicData(*InDynamicData)
+	#endif
 	{
 	}
 };
@@ -554,8 +561,8 @@ void LogDebugMessageLine(const ANSICHAR* File, int32 Line, const FName& InLogNam
 			((::ELogVerbosity::Verbosity & ::ELogVerbosity::VerbosityMask) <= ::ELogVerbosity::COMPILED_IN_MINIMUM_VERBOSITY && \
 			 (::ELogVerbosity::Verbosity & ::ELogVerbosity::VerbosityMask) <= UE_PRIVATE_LOG_CATEGORY(CategoryName).CompileTimeVerbosity)) \
 		{ \
-			static ::UE::Logging::Private::FStaticLogDynamicData LOG_Dynamic; \
-			static constexpr ::UE::Logging::Private::FStaticLogRecord LOG_Static UE_PRIVATE_LOGFMT_AGGREGATE(UTF8TEXT(Format), __builtin_FILE(), __builtin_LINE(), ::ELogVerbosity::Verbosity, LOG_Dynamic); \
+			UE_IF(LOGTRACE_ENABLED, static ::UE::Logging::Private::FStaticLogDynamicData LOG_Dynamic, ); \
+			static constexpr ::UE::Logging::Private::FStaticLogRecord LOG_Static UE_PRIVATE_LOGFMT_AGGREGATE(UTF8TEXT(Format), __builtin_FILE(), __builtin_LINE(), ::ELogVerbosity::Verbosity, UE_IF(LOGTRACE_ENABLED, &LOG_Dynamic, nullptr)); \
 			UE_PRIVATE_LOGFMT_LOG_IF_ACTIVE(Condition, CategoryName, Verbosity, LOG_Static, ##__VA_ARGS__); \
 		} \
 	} \
@@ -568,8 +575,8 @@ void LogDebugMessageLine(const ANSICHAR* File, int32 Line, const FName& InLogNam
 			((::ELogVerbosity::Verbosity & ::ELogVerbosity::VerbosityMask) <= ::ELogVerbosity::COMPILED_IN_MINIMUM_VERBOSITY && \
 			 (::ELogVerbosity::Verbosity & ::ELogVerbosity::VerbosityMask) <= UE_PRIVATE_LOG_CATEGORY(CategoryName).CompileTimeVerbosity)) \
 		{ \
-			static ::UE::Logging::Private::FStaticLogDynamicData LOG_Dynamic; \
-			static constexpr ::UE::Logging::Private::FStaticLocalizedLogRecord LOG_Static UE_PRIVATE_LOGFMT_AGGREGATE(TEXT(Namespace), TEXT(Key), UTF8TEXT(Format), __builtin_FILE(), __builtin_LINE(), ::ELogVerbosity::Verbosity, LOG_Dynamic); \
+			UE_IF(LOGTRACE_ENABLED, static ::UE::Logging::Private::FStaticLogDynamicData LOG_Dynamic, ); \
+			static constexpr ::UE::Logging::Private::FStaticLocalizedLogRecord LOG_Static UE_PRIVATE_LOGFMT_AGGREGATE(TEXT(Namespace), TEXT(Key), UTF8TEXT(Format), __builtin_FILE(), __builtin_LINE(), ::ELogVerbosity::Verbosity, UE_IF(LOGTRACE_ENABLED, &LOG_Dynamic, nullptr)); \
 			UE_PRIVATE_LOGFMT_LOG_IF_ACTIVE(Condition, CategoryName, Verbosity, LOG_Static, ##__VA_ARGS__); \
 		} \
 	} \
