Implemented an easy-to-use mechanism for deprecating existing cvars.

#Synced-CL 51264320

diff --git a/Engine/Source/Runtime/Core/Private/HAL/ConsoleManager.cpp b/Engine/Source/Runtime/Core/Private/HAL/ConsoleManager.cpp
--- a/Engine/Source/Runtime/Core/Private/HAL/ConsoleManager.cpp
+++ b/Engine/Source/Runtime/Core/Private/HAL/ConsoleManager.cpp
@@ -4741,7 +4741,7 @@ protected:
 			return false;
 		}
 
-		if (!bHasMessagedForUsage)
+		if (!bHasMessagedForUsage && !FCVarDisableDeprecationWarningScope::IsDisabled() )
 		{
 			if (UsageBehavior != EShadowCVarBehavior::NoMessaging)
 			{
@@ -5317,23 +5317,70 @@ static FString MakeDeprecatedHelp(const TCHAR* Name, const TCHAR* DeprecatedAtVe
 	return Message;
 }
 
-template<typename T>
-void CreateInPlaceDeprecatedVar(const TCHAR* Name, const T& DefaultValue, const TCHAR* DeprecatedAtVersion, EShadowCVarBehavior UsageBehavior, const TCHAR* DeprecationMessage)
+static thread_local int32 GDisabledCount = 0;
+
+FCVarDisableDeprecationWarningScope::FCVarDisableDeprecationWarningScope(const bool bInCondition) 
+	: bCondition(bInCondition)
+{
+	if (bCondition)
+	{
+		++GDisabledCount;
+	}
+}
+FCVarDisableDeprecationWarningScope::~FCVarDisableDeprecationWarningScope() 
+{
+	if (bCondition)
+	{
+		--GDisabledCount;
+	}
+}
+
+bool FCVarDisableDeprecationWarningScope::IsDisabled()
+{ 
+	return GDisabledCount > 0; 
+}
+
+template<typename T, typename CreateVarFuncType>
+IConsoleVariable* CreateInPlaceDeprecatedVar(const FCVarDeprecationArgs& DeprecationArgs, const TCHAR* Name, uint32 Flags, CreateVarFuncType CreateVarFunc)
 {
 #if !NO_CVARS
 	// create a backing unregistered cvar with a dummy name, but this name won't show up since Unregistered
 	FString BackingName = FString(Name) + "_dep";
-	IConsoleVariable* BackingCVar = new FConsoleVariable<T>(DefaultValue, *MakeDeprecatedHelp(Name, DeprecatedAtVersion, DeprecationMessage), (EConsoleVariableFlags)(ECVF_Default | ECVF_Unregistered));
+	IConsoleVariable* BackingCVar = CreateVarFunc(*MakeDeprecatedHelp(Name, DeprecationArgs.Version, DeprecationArgs.Message), (EConsoleVariableFlags)(Flags| ECVF_Unregistered));
 	GetManager().AddConsoleObject(*BackingName, BackingCVar);
 	
 	// create the shadow cvar that shadows the backing cvar, so we can use all the existing logic
-	FConsoleVariableDeprecatedNoReplace* DeprecatedCVar = new FConsoleVariableDeprecatedNoReplace(*BackingName, DeprecatedAtVersion, UsageBehavior, DeprecationMessage);
+	FConsoleVariableDeprecatedNoReplace* DeprecatedCVar = new FConsoleVariableDeprecatedNoReplace(*BackingName, DeprecationArgs.Version, DeprecationArgs.UsageBehavior, DeprecationArgs.Message);
 	// don't let the AddCosoleObject throw an error by calling GetFlags()
 	GetManager().AddShadowConsoleObject(Name, DeprecatedCVar);
 
+	return BackingCVar;
+#else
+	return nullptr;
 #endif
 }
 
+template<typename T>
+IConsoleVariable* CreateInPlaceDeprecatedVar(const FCVarDeprecationArgs& DeprecationArgs, const TCHAR* Name, const T& DefaultValue, uint32 Flags)
+{
+	return CreateInPlaceDeprecatedVar<T>(DeprecationArgs, Name, Flags, 
+		[&](const FString& DeprecationHelp, EConsoleVariableFlags ConsoleVariableFlags)
+	{ 
+		return new FConsoleVariable<T>(DefaultValue, *DeprecationHelp, ConsoleVariableFlags);
+	});
+}
+
+template<typename T>
+TAutoConsoleVariable<T>::TAutoConsoleVariable(FCVarDeprecationArgs DeprecationArgs, const TCHAR* Name, const T& DefaultValue, const TCHAR* Help, uint32 Flags)
+	: TAutoConsoleVariable(CreateInPlaceDeprecatedVar(DeprecationArgs, Name, DefaultValue, Flags))
+{
+}	
+
+template CORE_API TAutoConsoleVariable<int32>::TAutoConsoleVariable(FCVarDeprecationArgs DeprecationArgs, const TCHAR* Name, const int32& DefaultValue, const TCHAR* Help, uint32 Flags);
+template CORE_API TAutoConsoleVariable<float>::TAutoConsoleVariable(FCVarDeprecationArgs DeprecationArgs, const TCHAR* Name, const float& DefaultValue, const TCHAR* Help, uint32 Flags);
+template CORE_API TAutoConsoleVariable<bool>::TAutoConsoleVariable(FCVarDeprecationArgs DeprecationArgs, const TCHAR* Name, const bool& DefaultValue, const TCHAR* Help, uint32 Flags);
+template CORE_API TAutoConsoleVariable<FString>::TAutoConsoleVariable(FCVarDeprecationArgs DeprecationArgs, const TCHAR* Name, const FString& DefaultValue, const TCHAR* Help, uint32 Flags);
+
 template<typename T, typename CmdType>
 void CreateInPlaceDeprecatedCommand(const TCHAR* Name, const CmdType& Command, const TCHAR* DeprecatedAtVersion, EShadowCVarBehavior UsageBehavior, const TCHAR* DeprecationMessage)
 {
@@ -5373,31 +5420,17 @@ void CreateInPlaceDeprecatedCommandExec(const TCHAR* Name, const TCHAR* Deprecat
 #endif
 }
 
-// because we create a variable whose type is not in the header, we have to instantiate each possible type here
-template<>
-TAutoConsoleVariableDeprecated_NoReplace<int32>::TAutoConsoleVariableDeprecated_NoReplace(const TCHAR* Name, const int32& DefaultValue, const TCHAR* DeprecatedAtVersion, EShadowCVarBehavior UsageBehavior, const TCHAR* DeprecationMessage)
-{
-	CreateInPlaceDeprecatedVar(Name, DefaultValue, DeprecatedAtVersion, UsageBehavior, DeprecationMessage);
-}
-
-template<>
-TAutoConsoleVariableDeprecated_NoReplace<float>::TAutoConsoleVariableDeprecated_NoReplace(const TCHAR* Name, const float& DefaultValue, const TCHAR* DeprecatedAtVersion, EShadowCVarBehavior UsageBehavior, const TCHAR* DeprecationMessage)
-{
-	CreateInPlaceDeprecatedVar(Name, DefaultValue, DeprecatedAtVersion, UsageBehavior, DeprecationMessage);
-}
-
-template<>
-TAutoConsoleVariableDeprecated_NoReplace<bool>::TAutoConsoleVariableDeprecated_NoReplace(const TCHAR* Name, const bool& DefaultValue, const TCHAR* DeprecatedAtVersion, EShadowCVarBehavior UsageBehavior, const TCHAR* DeprecationMessage)
-{
-	CreateInPlaceDeprecatedVar(Name, DefaultValue, DeprecatedAtVersion, UsageBehavior, DeprecationMessage);
-}
-
-template<>
-TAutoConsoleVariableDeprecated_NoReplace<FString>::TAutoConsoleVariableDeprecated_NoReplace(const TCHAR* Name, const FString& DefaultValue, const TCHAR* DeprecatedAtVersion, EShadowCVarBehavior UsageBehavior, const TCHAR* DeprecationMessage)
+template<typename T>
+TAutoConsoleVariableDeprecated_NoReplace<T>::TAutoConsoleVariableDeprecated_NoReplace(const TCHAR* Name, const T& DefaultValue, const TCHAR* DeprecatedAtVersion, EShadowCVarBehavior UsageBehavior, const TCHAR* DeprecationMessage)
 {
-	CreateInPlaceDeprecatedVar(Name, DefaultValue, DeprecatedAtVersion, UsageBehavior, DeprecationMessage);
+	CreateInPlaceDeprecatedVar(FCVarDeprecationArgs{ DeprecatedAtVersion, UsageBehavior, DeprecationMessage }, Name, DefaultValue, ECVF_Default);
 }
 
+// because we create a variable whose type is not in the header, we have to instantiate each possible type here
+template CORE_API TAutoConsoleVariableDeprecated_NoReplace<int32>::TAutoConsoleVariableDeprecated_NoReplace(const TCHAR* Name, const int32& DefaultValue, const TCHAR* DeprecatedAtVersion, EShadowCVarBehavior UsageBehavior, const TCHAR* DeprecationMessage);
+template CORE_API TAutoConsoleVariableDeprecated_NoReplace<float>::TAutoConsoleVariableDeprecated_NoReplace(const TCHAR* Name, const float& DefaultValue, const TCHAR* DeprecatedAtVersion, EShadowCVarBehavior UsageBehavior, const TCHAR* DeprecationMessage);
+template CORE_API TAutoConsoleVariableDeprecated_NoReplace<bool>::TAutoConsoleVariableDeprecated_NoReplace(const TCHAR* Name, const bool& DefaultValue, const TCHAR* DeprecatedAtVersion, EShadowCVarBehavior UsageBehavior, const TCHAR* DeprecationMessage);
+template CORE_API TAutoConsoleVariableDeprecated_NoReplace<FString>::TAutoConsoleVariableDeprecated_NoReplace(const TCHAR* Name, const FString& DefaultValue, const TCHAR* DeprecatedAtVersion, EShadowCVarBehavior UsageBehavior, const TCHAR* DeprecationMessage);
 
 FAutoConsoleCommandDeprecated_NoReplace::FAutoConsoleCommandDeprecated_NoReplace(const TCHAR* Name, const TCHAR* DeprecatedAtVersion, EShadowCVarBehavior UsageBehavior, const TCHAR* DeprecationMessage)
 {
@@ -5481,6 +5514,16 @@ namespace ConsoleVariablePlatformMapping
 
 #if !NO_CVARS
 
+template<typename T>
+IConsoleVariable* CreateInPlaceDeprecatedVarRef(FCVarDeprecationArgs DeprecationArgs, const TCHAR* Name, T& RefValue, const TCHAR* Help, uint32 Flags)
+{
+	return CreateInPlaceDeprecatedVar<T>(DeprecationArgs, Name, Flags, 
+		[&](const FString& DeprecationHelp, EConsoleVariableFlags ConsoleVariableFlags)
+	{ 
+		return IConsoleManager::Get().RegisterConsoleVariableRef(Name, RefValue, *DeprecationHelp, ConsoleVariableFlags);
+	});
+}
+
 FAutoConsoleVariable::FAutoConsoleVariable(const TCHAR* Name, bool DefaultValue, const TCHAR* Help, uint32 Flags)
 	: FAutoConsoleObject(IConsoleManager::Get().RegisterConsoleVariable(Name, DefaultValue, Help, Flags))
 {
@@ -5530,49 +5573,91 @@ FAutoConsoleVariableRef::FAutoConsoleVariableRef(const TCHAR* Name, int32& RefVa
 {
 }
 
+FAutoConsoleVariableRef::FAutoConsoleVariableRef(FCVarDeprecationArgs DeprecationArgs, const TCHAR* Name, int32& RefValue, const TCHAR* Help, uint32 Flags)
+	: FAutoConsoleObject(CreateInPlaceDeprecatedVarRef(DeprecationArgs, Name, RefValue, Help, Flags))
+{
+}
+
 FAutoConsoleVariableRef::FAutoConsoleVariableRef(const TCHAR* Name, float& RefValue, const TCHAR* Help, uint32 Flags)
 	: FAutoConsoleObject(IConsoleManager::Get().RegisterConsoleVariableRef(Name, RefValue, Help, Flags))
 {
 }
 
+FAutoConsoleVariableRef::FAutoConsoleVariableRef(FCVarDeprecationArgs DeprecationArgs, const TCHAR* Name, float& RefValue, const TCHAR* Help, uint32 Flags)
+	: FAutoConsoleObject(CreateInPlaceDeprecatedVarRef(DeprecationArgs, Name, RefValue, Help, Flags))
+{
+}
 FAutoConsoleVariableRef::FAutoConsoleVariableRef(const TCHAR* Name, bool& RefValue, const TCHAR* Help, uint32 Flags)
 	: FAutoConsoleObject(IConsoleManager::Get().RegisterConsoleVariableRef(Name, RefValue, Help, Flags))
 {
 }
 
+FAutoConsoleVariableRef::FAutoConsoleVariableRef(FCVarDeprecationArgs DeprecationArgs, const TCHAR* Name, bool& RefValue, const TCHAR* Help, uint32 Flags)
+	: FAutoConsoleObject(CreateInPlaceDeprecatedVarRef(DeprecationArgs, Name, RefValue, Help, Flags))
+{
+}
 FAutoConsoleVariableRef::FAutoConsoleVariableRef(const TCHAR* Name, FString& RefValue, const TCHAR* Help, uint32 Flags)
 	: FAutoConsoleObject(IConsoleManager::Get().RegisterConsoleVariableRef(Name, RefValue, Help, Flags))
 {
 }
 
+FAutoConsoleVariableRef::FAutoConsoleVariableRef(FCVarDeprecationArgs DeprecationArgs, const TCHAR* Name, FString& RefValue, const TCHAR* Help, uint32 Flags)
+	: FAutoConsoleObject(CreateInPlaceDeprecatedVarRef(DeprecationArgs, Name, RefValue, Help, Flags))
+{
+}
 FAutoConsoleVariableRef::FAutoConsoleVariableRef(const TCHAR* Name, FName& RefValue, const TCHAR* Help, uint32 Flags)
 	: FAutoConsoleObject(IConsoleManager::Get().RegisterConsoleVariableRef(Name, RefValue, Help, Flags))
 {
 }
 
+FAutoConsoleVariableRef::FAutoConsoleVariableRef(FCVarDeprecationArgs DeprecationArgs, const TCHAR* Name, FName& RefValue, const TCHAR* Help, uint32 Flags)
+	: FAutoConsoleObject(CreateInPlaceDeprecatedVarRef(DeprecationArgs, Name, RefValue, Help, Flags))
+{
+}
 FAutoConsoleVariableRef::FAutoConsoleVariableRef(const TCHAR* Name, int32& RefValue, const TCHAR* Help, const FConsoleVariableDelegate& Callback, uint32 Flags)
 	: FAutoConsoleObject(IConsoleManager::Get().RegisterConsoleVariableRef(Name, RefValue, Help, Flags))
 {
 	AsVariable()->SetOnChangedCallback(Callback);
 }
 
+FAutoConsoleVariableRef::FAutoConsoleVariableRef(FCVarDeprecationArgs DeprecationArgs, const TCHAR* Name, int32& RefValue, const TCHAR* Help, const FConsoleVariableDelegate& Callback, uint32 Flags)
+	: FAutoConsoleObject(CreateInPlaceDeprecatedVarRef(DeprecationArgs, Name, RefValue, Help, Flags))
+{
+	AsVariable()->SetOnChangedCallback(Callback);
+}
+
 FAutoConsoleVariableRef::FAutoConsoleVariableRef(const TCHAR* Name, float& RefValue, const TCHAR* Help, const FConsoleVariableDelegate& Callback, uint32 Flags)
 	: FAutoConsoleObject(IConsoleManager::Get().RegisterConsoleVariableRef(Name, RefValue, Help, Flags))
 {
 	AsVariable()->SetOnChangedCallback(Callback);
 }
+FAutoConsoleVariableRef::FAutoConsoleVariableRef(FCVarDeprecationArgs DeprecationArgs, const TCHAR* Name, float& RefValue, const TCHAR* Help, const FConsoleVariableDelegate& Callback, uint32 Flags)
+	: FAutoConsoleObject(CreateInPlaceDeprecatedVarRef(DeprecationArgs, Name, RefValue, Help, Flags))
+{
+	AsVariable()->SetOnChangedCallback(Callback);
+}
 
 FAutoConsoleVariableRef::FAutoConsoleVariableRef(const TCHAR* Name, bool& RefValue, const TCHAR* Help, const FConsoleVariableDelegate& Callback, uint32 Flags)
 	: FAutoConsoleObject(IConsoleManager::Get().RegisterConsoleVariableRef(Name, RefValue, Help, Flags))
 {
 	AsVariable()->SetOnChangedCallback(Callback);
 }
+FAutoConsoleVariableRef::FAutoConsoleVariableRef(FCVarDeprecationArgs DeprecationArgs, const TCHAR* Name, bool& RefValue, const TCHAR* Help, const FConsoleVariableDelegate& Callback, uint32 Flags)
+	: FAutoConsoleObject(CreateInPlaceDeprecatedVarRef(DeprecationArgs, Name, RefValue, Help, Flags))
+{
+	AsVariable()->SetOnChangedCallback(Callback);
+}
 
 FAutoConsoleVariableRef::FAutoConsoleVariableRef(const TCHAR* Name, FString& RefValue, const TCHAR* Help, const FConsoleVariableDelegate& Callback, uint32 Flags)
 	: FAutoConsoleObject(IConsoleManager::Get().RegisterConsoleVariableRef(Name, RefValue, Help, Flags))
 {
 	AsVariable()->SetOnChangedCallback(Callback);
 }
+FAutoConsoleVariableRef::FAutoConsoleVariableRef(FCVarDeprecationArgs DeprecationArgs, const TCHAR* Name, FString& RefValue, const TCHAR* Help, const FConsoleVariableDelegate& Callback, uint32 Flags)
+	: FAutoConsoleObject(CreateInPlaceDeprecatedVarRef(DeprecationArgs, Name, RefValue, Help, Flags))
+{
+	AsVariable()->SetOnChangedCallback(Callback);
+}
 
 FAutoConsoleCommand::FAutoConsoleCommand(const TCHAR* Name, const TCHAR* Help, const FConsoleCommandDelegate& Command, uint32 Flags)
 	: FAutoConsoleObject(IConsoleManager::Get().RegisterConsoleCommand(Name, Help, Command, Flags))
diff --git a/Engine/Source/Runtime/Core/Public/HAL/IConsoleManager.h b/Engine/Source/Runtime/Core/Public/HAL/IConsoleManager.h
--- a/Engine/Source/Runtime/Core/Public/HAL/IConsoleManager.h
+++ b/Engine/Source/Runtime/Core/Public/HAL/IConsoleManager.h
@@ -1589,6 +1589,50 @@ public:
 };
 #endif
 
+
+enum class EShadowCVarBehavior : uint8
+{
+	// skip all messaging
+	NoMessaging,
+	// log a warning on first usage
+	Warn,
+	// log an error on first usage
+	Error,
+	// throw an ensure on first usage
+	Ensure,
+	// throw an assertion on first usage
+	Assert,
+};
+
+/**
+ * Disable deprecation warnings for access to cvars for the duration of the current scope, enabling internal use of deprecated cvars without generating warnings / ensures.
+ * Analogous to the use case for the pragmas for deprecated code.
+ */
+struct FCVarDisableDeprecationWarningScope
+{
+	/**
+	 * if bInCondition is false, the scope does nothing such that it can be conditionally used.
+	 */
+	CORE_API FCVarDisableDeprecationWarningScope(bool bInCondition = true);
+	CORE_API ~FCVarDisableDeprecationWarningScope();
+
+	CORE_API static bool IsDisabled();
+
+private:
+	const bool bCondition = true;
+};
+
+struct FCVarDeprecationArgs
+{
+	// The first UE version when this was deprecated
+	const TCHAR* Version = TEXT("");
+	// How to handle uses of the deprecated cvar  - defaults to ensure
+	EShadowCVarBehavior UsageBehavior = EShadowCVarBehavior::Ensure;
+	 // Optional deprecation message - defaults to empty / auto-generated.
+	const TCHAR* Message = TEXT("");
+};
+
+
 #if !NO_CVARS
 /**
  * Autoregistering float, int, bool, FString REF variable class...this changes that value when the console variable is changed. 
@@ -1603,6 +1647,7 @@ public:
 	 * @param Flags bitmask combined from EConsoleVariableFlags
 	 */
 	CORE_API FAutoConsoleVariableRef(const TCHAR* Name, int32& RefValue, const TCHAR* Help, uint32 Flags = ECVF_Default);
+	CORE_API FAutoConsoleVariableRef(FCVarDeprecationArgs DeprecationArgs, const TCHAR* Name, int32& RefValue, const TCHAR* Help, uint32 Flags = ECVF_Default);
 	/**
 	 * Create a reference to a float console variable
 	 * @param Name must not be 0
@@ -1610,6 +1655,7 @@ public:
 	 * @param Flags bitmask combined from EConsoleVariableFlags
 	 */
 	CORE_API FAutoConsoleVariableRef(const TCHAR* Name, float& RefValue, const TCHAR* Help, uint32 Flags = ECVF_Default);
+	CORE_API FAutoConsoleVariableRef(FCVarDeprecationArgs DeprecationArgs, const TCHAR* Name, float& RefValue, const TCHAR* Help, uint32 Flags = ECVF_Default);
 	/**
 	 * Create a reference to a bool console variable
 	 * @param Name must not be 0
@@ -1617,6 +1663,7 @@ public:
 	 * @param Flags bitmask combined from EConsoleVariableFlags
 	 */
 	CORE_API FAutoConsoleVariableRef(const TCHAR* Name, bool& RefValue, const TCHAR* Help, uint32 Flags = ECVF_Default);
+	CORE_API FAutoConsoleVariableRef(FCVarDeprecationArgs DeprecationArgs, const TCHAR* Name, bool& RefValue, const TCHAR* Help, uint32 Flags = ECVF_Default);
 	/**
 	 * Create a reference to a FString console variable
 	 * @param Name must not be 0
@@ -1624,6 +1671,7 @@ public:
 	 * @param Flags bitmask combined from EConsoleVariableFlags
 	 */
 	CORE_API FAutoConsoleVariableRef(const TCHAR* Name, FString& RefValue, const TCHAR* Help, uint32 Flags = ECVF_Default);
+	CORE_API FAutoConsoleVariableRef(FCVarDeprecationArgs DeprecationArgs, const TCHAR* Name, FString& RefValue, const TCHAR* Help, uint32 Flags = ECVF_Default);
 
 	/**
 	 * Create a reference to a FName console variable
@@ -1632,6 +1680,7 @@ public:
 	 * @param Flags bitmask combined from EConsoleVariableFlags
 	 */
 	CORE_API FAutoConsoleVariableRef(const TCHAR* Name, FName& RefValue, const TCHAR* Help, uint32 Flags = ECVF_Default);
+	CORE_API FAutoConsoleVariableRef(FCVarDeprecationArgs DeprecationArgs, const TCHAR* Name, FName& RefValue, const TCHAR* Help, uint32 Flags = ECVF_Default);
 
 	/**
 	 * Create a reference to a int console variable
@@ -1641,6 +1690,7 @@ public:
 	 * @param Flags bitmask combined from EConsoleVariableFlags
 	 */
 	CORE_API FAutoConsoleVariableRef(const TCHAR* Name, int32& RefValue, const TCHAR* Help, const FConsoleVariableDelegate& Callback, uint32 Flags = ECVF_Default);
+	CORE_API FAutoConsoleVariableRef(FCVarDeprecationArgs DeprecationArgs, const TCHAR* Name, int32& RefValue, const TCHAR* Help, const FConsoleVariableDelegate& Callback, uint32 Flags = ECVF_Default);
 
 	/**
 	 * Create a reference to a float console variable
@@ -1650,6 +1700,7 @@ public:
 	 * @param Flags bitmask combined from EConsoleVariableFlags
 	 */
 	CORE_API FAutoConsoleVariableRef(const TCHAR* Name, float& RefValue, const TCHAR* Help, const FConsoleVariableDelegate& Callback, uint32 Flags = ECVF_Default);
+	CORE_API FAutoConsoleVariableRef(FCVarDeprecationArgs DeprecationArgs, const TCHAR* Name, float& RefValue, const TCHAR* Help, const FConsoleVariableDelegate& Callback, uint32 Flags = ECVF_Default);
 
 	/**
 	 * Create a reference to a bool console variable
@@ -1659,6 +1710,7 @@ public:
 	 * @param Flags bitmask combined from EConsoleVariableFlags
 	 */
 	CORE_API FAutoConsoleVariableRef(const TCHAR* Name, bool& RefValue, const TCHAR* Help, const FConsoleVariableDelegate& Callback, uint32 Flags = ECVF_Default);
+	CORE_API FAutoConsoleVariableRef(FCVarDeprecationArgs DeprecationArgs, const TCHAR* Name, bool& RefValue, const TCHAR* Help, const FConsoleVariableDelegate& Callback, uint32 Flags = ECVF_Default);
 
 	/**
 	 * Create a reference to a FString console variable
@@ -1668,6 +1720,7 @@ public:
 	 * @param Flags bitmask combined from EConsoleVariableFlags
 	 */
 	CORE_API FAutoConsoleVariableRef(const TCHAR* Name, FString& RefValue, const TCHAR* Help, const FConsoleVariableDelegate& Callback, uint32 Flags = ECVF_Default);
+	CORE_API FAutoConsoleVariableRef(FCVarDeprecationArgs DeprecationArgs, const TCHAR* Name, FString& RefValue, const TCHAR* Help, const FConsoleVariableDelegate& Callback, uint32 Flags = ECVF_Default);
 
 	virtual ~FAutoConsoleVariableRef()
 	{
@@ -1847,6 +1900,8 @@ public:
 	{
 	}
 
+	CORE_API TAutoConsoleVariable(FCVarDeprecationArgs DeprecationArgs, const TCHAR* Name, const T& DefaultValue, const TCHAR* Help, uint32 Flags = ECVF_Default);
+	
 	/**
 	 * Create a float, int or string console variable
 	 * @param Name must not be 0
@@ -1860,6 +1915,12 @@ public:
 		AsVariable()->SetOnChangedCallback(Callback);
 	}
 
+	FORCENOINLINE TAutoConsoleVariable(FCVarDeprecationArgs DeprecationArgs, const TCHAR* Name, const T& DefaultValue, const TCHAR* Help, const FConsoleVariableDelegate& Callback, uint32 Flags = ECVF_Default)
+		: TAutoConsoleVariable(DeprecationArgs,Name, DefaultValue, Help, Flags)
+	{
+		AsVariable()->SetOnChangedCallback(Callback);
+	}
+
 	virtual ~TAutoConsoleVariable()
 	{
 		AsVariable()->ClearOnChangedCallback();
@@ -2248,20 +2309,6 @@ public:
 #endif
 
 
-enum class EShadowCVarBehavior : uint8
-{
-	// skip all messaging
-	NoMessaging,
-	// log a warning on first usage
-	Warn,
-	// log an error on first usage
-	Error,
-	// throw an ensure on first usage
-	Ensure,
-	// throw an assertion on first usage
-	Assert,
-};
-
 /**
  * A class that can be used to shadow (alias) a CVar, where both names are valid to use.
  * To deprecate a CVar, use FAutoConsoleVariableDeprecated.
