Mutex: Added IsLockedByCurrentThread() to the recursive mutex types

#Synced-CL 50967057

diff --git a/Engine/Source/Runtime/Core/Private/Async/RecursiveMutex.cpp b/Engine/Source/Runtime/Core/Private/Async/RecursiveMutex.cpp
--- a/Engine/Source/Runtime/Core/Private/Async/RecursiveMutex.cpp
+++ b/Engine/Source/Runtime/Core/Private/Async/RecursiveMutex.cpp
@@ -10,6 +10,14 @@
 namespace UE
 {
 
+bool FRecursiveMutex::IsLockedByCurrentThread() const
+{
+	UE_AUTORTFM_REPORT_HAZARD_IF_CLOSED(RecursiveMutex);
+
+	const uint32 CurrentThreadId = FPlatformTLS::GetCurrentThreadId();
+	return ThreadId.load(std::memory_order_relaxed) == CurrentThreadId;
+}
+
 bool FRecursiveMutex::TryLock()
 {
 	UE_AUTORTFM_REPORT_HAZARD_IF_CLOSED(RecursiveMutex);
diff --git a/Engine/Source/Runtime/Core/Private/Async/RecursiveWordMutex.cpp b/Engine/Source/Runtime/Core/Private/Async/RecursiveWordMutex.cpp
--- a/Engine/Source/Runtime/Core/Private/Async/RecursiveWordMutex.cpp
+++ b/Engine/Source/Runtime/Core/Private/Async/RecursiveWordMutex.cpp
@@ -6,6 +6,12 @@
 namespace UE
 {
 
+bool FRecursiveWordMutex::IsLockedByCurrentThread() const
+{
+	const uint32 CurrentThreadId = FPlatformTLS::GetCurrentThreadId();
+	return ThreadId.load(std::memory_order_relaxed) == CurrentThreadId;
+}
+
 bool FRecursiveWordMutex::TryLock()
 {
 	const uint32 CurrentThreadId = FPlatformTLS::GetCurrentThreadId();
diff --git a/Engine/Source/Runtime/Core/Private/Async/SharedRecursiveMutex.cpp b/Engine/Source/Runtime/Core/Private/Async/SharedRecursiveMutex.cpp
--- a/Engine/Source/Runtime/Core/Private/Async/SharedRecursiveMutex.cpp
+++ b/Engine/Source/Runtime/Core/Private/Async/SharedRecursiveMutex.cpp
@@ -77,6 +77,12 @@ inline const void* FSharedRecursiveMutex::GetSharedLockAddress() const
 	return (const uint8*)&State + 1;
 }
 
+bool FSharedRecursiveMutex::IsLockedByCurrentThread() const
+{
+	const uint32 CurrentThreadId = FPlatformTLS::GetCurrentThreadId();
+	return ThreadId.load(std::memory_order_relaxed) == CurrentThreadId;
+}
+
 bool FSharedRecursiveMutex::TryLock()
 {
 	const uint32 CurrentThreadId = FPlatformTLS::GetCurrentThreadId();
diff --git a/Engine/Source/Runtime/Core/Public/Async/RecursiveMutex.h b/Engine/Source/Runtime/Core/Public/Async/RecursiveMutex.h
--- a/Engine/Source/Runtime/Core/Public/Async/RecursiveMutex.h
+++ b/Engine/Source/Runtime/Core/Public/Async/RecursiveMutex.h
@@ -30,6 +30,7 @@ public:
 		return !!(State.load(std::memory_order_relaxed) & LockCountMask);
 	}
 
+	[[nodiscard]] UE_API bool IsLockedByCurrentThread() const;
 	[[nodiscard]] UE_API bool TryLock();
 	UE_API void Lock();
 	UE_API void Unlock();
diff --git a/Engine/Source/Runtime/Core/Public/Async/RecursiveWordMutex.h b/Engine/Source/Runtime/Core/Public/Async/RecursiveWordMutex.h
--- a/Engine/Source/Runtime/Core/Public/Async/RecursiveWordMutex.h
+++ b/Engine/Source/Runtime/Core/Public/Async/RecursiveWordMutex.h
@@ -26,6 +26,8 @@ public:
 	{
 		return Mutex.IsLocked();
 	}
+
+	[[nodiscard]] UE_API bool IsLockedByCurrentThread() const;
 	[[nodiscard]] UE_API bool TryLock();
 	UE_API void Lock();
 	UE_API void Unlock();
diff --git a/Engine/Source/Runtime/Core/Public/Async/SharedRecursiveMutex.h b/Engine/Source/Runtime/Core/Public/Async/SharedRecursiveMutex.h
--- a/Engine/Source/Runtime/Core/Public/Async/SharedRecursiveMutex.h
+++ b/Engine/Source/Runtime/Core/Public/Async/SharedRecursiveMutex.h
@@ -49,6 +49,7 @@ public:
 		return !!(State.load(std::memory_order_relaxed) & LockCountMask);
 	}
 
+	[[nodiscard]] UE_API bool IsLockedByCurrentThread() const;
 	[[nodiscard]] UE_API bool TryLock();
 	UE_API void Lock();
 	UE_API void Unlock();
