AssetRegistry: AssetDataGatherer: When using junctions, support finding the proper USN journal file for the physical drive the junction is on. This allows the discovery cache to be used for junctions to other drives.
#rnx

#Synced-CL 51202900

diff --git a/Engine/Source/Runtime/Core/Private/Windows/WindowsPlatformFileJournal.cpp b/Engine/Source/Runtime/Core/Private/Windows/WindowsPlatformFileJournal.cpp
--- a/Engine/Source/Runtime/Core/Private/Windows/WindowsPlatformFileJournal.cpp
+++ b/Engine/Source/Runtime/Core/Private/Windows/WindowsPlatformFileJournal.cpp
@@ -789,11 +789,78 @@ EFileJournalResult CreateJournalWindowsHandle(const TCHAR* VolumeOrPath, HANDLE&
 	return CreateJournalWindowsHandle(*OutVolumeName, OutHandle, OutError);
 }
 
+void GetVolumeNameOnDisk(const TCHAR* VolumeName, FStringBuilderBase& OutVolume)
+{
+	OutVolume.Reset();
+
+	constexpr int32 InitialCapacity = 64;
+	TStringBuilder<InitialCapacity> Buffer;
+
+	Buffer << VolumeName;
+	// CreateFile will treat a volume name "D:" as a relative path to CWD on the drive, and we
+	// want it to instead open the root path "D:\". Add on a trailing slash.
+	if (Buffer.Len() > 0 && !FPathViews::IsSeparator(Buffer.GetData()[Buffer.Len() - 1]))
+	{
+		Buffer << '\\';
+	}
+
+	HANDLE hFile = CreateFile(*Buffer,
+		FILE_READ_ATTRIBUTES, FILE_SHARE_READ | FILE_SHARE_WRITE | FILE_SHARE_DELETE,
+		NULL, OPEN_EXISTING,
+		FILE_FLAG_BACKUP_SEMANTICS, // Required for opening a directory, which Buffer should be.
+		NULL);
+	if (hFile == INVALID_HANDLE_VALUE)
+	{
+		OutVolume << VolumeName;
+		return;
+	}
+	ON_SCOPE_EXIT
+	{
+		CloseHandle(hFile);
+	};
+	
+	for (uint32 Length = (uint32)InitialCapacity - 1; ;)
+	{
+		if (Length + 1 > (uint32)Buffer.Len())
+		{
+			Buffer.AddUninitialized(Length + 1 - Buffer.Len());
+		}
+		Length = GetFinalPathNameByHandle(hFile, Buffer.GetData(),
+			Buffer.Len(), // BufferSize, which must include room for a null termination character
+			FILE_NAME_NORMALIZED);
+		// Returned value is the required text length for the path not counting the null termination character;
+		// only Min(Length+1, BufferSize) was written.
+		if (Length == 0)
+		{
+			OutVolume << VolumeName;
+			return;
+		}
+		if (Length < (uint32)Buffer.Len())
+		{
+			Buffer.RemoveSuffix(Buffer.Len() - Length);
+			break;
+		}
+	}
+	FStringView UNCPathPrefix(TEXT("\\\\?\\"));
+	FString FinalPathView(Buffer);
+	if (FinalPathView.StartsWith(UNCPathPrefix))
+	{
+		FinalPathView = FinalPathView.RightChop(UNCPathPrefix.Len());
+	}
+	FStringView OutVolumeView;
+	FStringView Remainder;
+	FPathViews::SplitVolumeSpecifier(FinalPathView, OutVolumeView, Remainder);
+	OutVolume << OutVolumeView;
+}
+
 EFileJournalResult CreateJournalWindowsHandle(const TCHAR* VolumeName, HANDLE& OutHandle, FString* OutError)
 {
 	OutHandle = NULL;
 
-	TStringBuilder<32> JournalName(InPlace, TEXT("\\\\.\\"), VolumeName);
+	TStringBuilder<64> PhysicalVolumeName;
+	GetVolumeNameOnDisk(VolumeName, PhysicalVolumeName);
+
+	TStringBuilder<32> JournalName(InPlace, TEXT("\\\\.\\"), *PhysicalVolumeName);
 	OutHandle = ::CreateFileW(
 		*JournalName,
 		FILE_TRAVERSE,
@@ -809,7 +876,7 @@ EFileJournalResult CreateJournalWindowsHandle(const TCHAR* VolumeName, HANDLE& O
 		{
 			uint32 LastError = FPlatformMisc::GetLastError();
 			*OutError = FString::Printf(TEXT("CreateFile failed to open for reading FileJournal on volume '%s'. GetLastError == %d (0x%08x) (%s)"),
-				VolumeName, LastError, LastError, *JournalGetLastErrorString(LastError));
+				*PhysicalVolumeName, LastError, LastError, *JournalGetLastErrorString(LastError));
 		}
 		return EFileJournalResult::FailedOpenJournal;
 	}
