[WinArm64] Initial work to support PGO using MSVC on ARM64 and ARM64EC.

Due to internal compiler errors, I've been unable to complete a PGO build on ARM64 or ARM64EC. These are preliminary changes, but there may be more necessary.

When building with -PGOProfile, x64 copies binaries to the root Binaries\Win64 directory. Both ARM64 and ARM64EC use a manifest to load DLLs from a subdirectory named arm64 and arm64ec respectively. When compiling for ARM64 or ARM64EC, we now copy the appropriate PGO files to those subdirectories so building for multiple architectures doesn't stomp the same files.

pgort140.dll is an ARM64X DLL so both ARM64 and ARM64EC are supported. The engine spawns pgosweep.exe in PGO_WriteFile, so both use the native ARM64 binary.

#rnx

#Synced-CL 51181113

diff --git a/Engine/Source/Programs/UnrealBuildTool/Platform/Windows/VCToolChain.cs b/Engine/Source/Programs/UnrealBuildTool/Platform/Windows/VCToolChain.cs
--- a/Engine/Source/Programs/UnrealBuildTool/Platform/Windows/VCToolChain.cs
+++ b/Engine/Source/Programs/UnrealBuildTool/Platform/Windows/VCToolChain.cs
@@ -2767,16 +2767,26 @@ namespace UnrealBuildTool
 			// copy PGO support binaries for Windows from $(VC_PGO_RunTime_Dir)
 			if (Target.bPGOProfile && Target.Platform.IsInGroup(UnrealPlatformGroup.Windows) && !Target.WindowsPlatform.Compiler.IsClang())
 			{
+				// Both ARM64 and ARM64EC use the same files. This requires the DLLs are all built as ARM64X.
 				string[] PGOFiles = {
 					"pgort140.dll",
 					"pgosweep.exe",
 				};
+
+				// Non-x64 files are placed in sub-directories. A manifest is used to load DLLs from this directory,
+				// and the EXEs just need to not stomp each other when building multiple architectures at once.
+				DirectoryReference DstDir = ExeDir;
+				if (!Target.Architecture.bIsX64)
+				{
+					DstDir = DirectoryReference.Combine(ExeDir, Target.Architecture.ToString());
+				}
+
 				foreach (string PGOFile in PGOFiles)
 				{
 					FileReference SrcFile = EnvVars.Architecture.bIsX64
 						? FileReference.Combine(EnvVars.ToolChainDir, "bin", "Hostx64", "x64", PGOFile)
 						: FileReference.Combine(EnvVars.ToolChainDir, "bin", "arm64", PGOFile);
-					FileReference DstFile = FileReference.Combine(ExeDir, PGOFile);
+					FileReference DstFile = FileReference.Combine(DstDir, PGOFile);
 					TargetFileToSourceFile[DstFile] = SrcFile;
 					RuntimeDependencies.Add(new RuntimeDependency(DstFile, StagedFileType.NonUFS));
 				}
diff --git a/Engine/Source/Runtime/Core/Private/Windows/WindowsPlatformMisc.cpp b/Engine/Source/Runtime/Core/Private/Windows/WindowsPlatformMisc.cpp
--- a/Engine/Source/Runtime/Core/Private/Windows/WindowsPlatformMisc.cpp
+++ b/Engine/Source/Runtime/Core/Private/Windows/WindowsPlatformMisc.cpp
@@ -885,18 +885,34 @@ void PGO_WriteFile()
 	} while (GetFileAttributesW(*OutputFilePath) != INVALID_FILE_ATTRIBUTES);
 
 	// Launch PGOSweep & wait for it to finish
+#if PLATFORM_WINDOWS_ARM64EC
+	FString PGOSweepPath = FPaths::Combine(ExeFolder, TEXT("arm64ec"), TEXT("pgosweep.exe"));
+#elif PLATFORM_CPU_ARM_FAMILY
+	FString PGOSweepPath = FPaths::Combine(ExeFolder, TEXT("arm64"), TEXT("pgosweep.exe"));
+#else
 	FString PGOSweepPath = FPaths::Combine(ExeFolder, TEXT("pgosweep.exe"));
+#endif
+	
 	FString CommandLine = FString::Printf(TEXT("/pid:%d \"%s\" \"%s\""), ::GetCurrentProcessId(), ExeFilePath, *OutputFilePath);
 
 	const bool bLaunchDetached = true;
 	const bool bLaunchHidden = true;
 	const bool bLaunchReallyHidden = bLaunchHidden;
 	FProcHandle ProcHandle = FPlatformProcess::CreateProc(*PGOSweepPath, *CommandLine, bLaunchDetached, bLaunchHidden, bLaunchReallyHidden, nullptr, 0, nullptr, nullptr, nullptr);
-	FPlatformProcess::WaitForProc(ProcHandle);
-	int32 ExitCode = 0;
-	if (FPlatformProcess::GetProcReturnCode(ProcHandle, &ExitCode))
+
+	if (ProcHandle.IsValid())
+	{
+		FPlatformProcess::WaitForProc(ProcHandle);
+		int32 ExitCode = 0;
+		if (FPlatformProcess::GetProcReturnCode(ProcHandle, &ExitCode))
+		{
+			UE_LOGF(LogWindows, Log, "pgosweep.exe exit code %d", ExitCode);
+		}
+		FPlatformProcess::CloseProc(ProcHandle);
+	}
+	else
 	{
-		UE_LOG(LogWindows, Log, TEXT("pgosweep.exe exit code %d"), ExitCode);
+		UE_LOGF(LogWindows, Error, "Failed to launch pgosweep.exe");
 	}
 }
 #endif //ENABLE_PGO_PROFILE && !__clang__ && !__INTEL_LLVM_COMPILER
