Add cvar csv.threadpriority to override the thread priority of the processing thread. This is settable at runtime
#tests tested on console

#Synced-CL 51245126

diff --git a/Engine/Source/Runtime/Core/Private/ProfilingDebugging/CsvProfiler.cpp b/Engine/Source/Runtime/Core/Private/ProfilingDebugging/CsvProfiler.cpp
--- a/Engine/Source/Runtime/Core/Private/ProfilingDebugging/CsvProfiler.cpp
+++ b/Engine/Source/Runtime/Core/Private/ProfilingDebugging/CsvProfiler.cpp
@@ -261,6 +261,37 @@ TAutoConsoleVariable<int32> CVarCsvStatNameValidation(
 	ECVF_Default
 );
 
+// Thread priority cvar (settable at runtime)
+// We declare these explicitly rather than just casting the cvar in case the enum changes in future
+const EThreadPriority GCsvThreadPriorities[] =
+{
+	EThreadPriority::TPri_Lowest,
+	EThreadPriority::TPri_BelowNormal,
+	EThreadPriority::TPri_SlightlyBelowNormal,
+	EThreadPriority::TPri_Normal,
+	EThreadPriority::TPri_AboveNormal,
+	EThreadPriority::TPri_Highest
+};
+
+const TCHAR* GCsvThreadPriorityNames[] =
+{
+	TEXT("TPri_Lowest"),
+	TEXT("TPri_BelowNormal"),
+	TEXT("TPri_SlightlyBelowNormal"),
+	TEXT("TPri_Normal"),
+	TEXT("TPri_AboveNormal"),
+	TEXT("TPri_Highest")
+};
+
+static int32 GCsvThreadPriorityIndex = 0; 
+FAutoConsoleVariableRef CVarCsvThreadPriority(
+	TEXT("csv.ThreadPriority"),
+	GCsvThreadPriorityIndex,
+	TEXT("Thread priority of the CSV profiler thread: 0=Lowest, 1=BelowNormal, 2=SlightlyBelowNormal, 3=Normal, 4=AboveNormal, 5=Highest\n")
+	TEXT("This is switchable at runtime"),
+	ECVF_Default);
+
+
 static bool GCsvUseProcessingThread = true;
 static int32 GCsvRepeatCount = 0;
 static int32 GCsvRepeatFrameCount = 0;
@@ -3207,22 +3238,23 @@ inline bool FCsvStreamWriter::HasReachedMaxFileBytes() const
 class FCsvProfilerProcessingThread : public FRunnable
 {
 	FThreadSafeCounter StopCounter;
-	EThreadPriority Priority;
+	EThreadPriority CurrentThreadPriority;
 
 public:
 	FCsvProfilerProcessingThread(FCsvProfiler& InCsvProfiler)
 		: CsvProfiler(InCsvProfiler)
 	{
-		Priority = TPri_Lowest;
 		uint64 AffinityMask = FPlatformAffinity::GetTaskGraphBackgroundTaskMask(); 
 #if CSV_PROFILER_ALLOW_DEBUG_FEATURES
 		if (FParse::Param(FCommandLine::Get(), TEXT("csvProfilerHighPriority")))
 		{
-			Priority = TPri_Highest;
+			GCsvThreadPriorityIndex = 5;
 			AffinityMask = FPlatformAffinity::GetTaskGraphThreadMask();
 		}
 #endif
-		Thread = FForkProcessHelper::CreateForkableThread(this, TEXT("CSVProfiler"), 0, Priority, AffinityMask);
+		UE_LOG(LogCsvProfiler, Display, TEXT("Starting CSV processing thread with priority %s"), GetThreadPriorityNameFromIndex(GCsvThreadPriorityIndex));
+		CurrentThreadPriority = GetThreadPriorityFromIndex(GCsvThreadPriorityIndex);
+		Thread = FForkProcessHelper::CreateForkableThread(this, TEXT("CSVProfiler"), 0, CurrentThreadPriority, AffinityMask);
 	}
 
 	virtual ~FCsvProfilerProcessingThread()
@@ -3264,6 +3296,7 @@ public:
 
 		while (StopCounter.GetValue() == 0)
 		{
+			UpdateThreadPriorityIfNeeded();
 			const float TimeBetweenUpdatesMS = CVarCsvProcessingThreadTimeBetweenUpdates.GetValueOnAnyThread();
 
 			// Pause the processing thread if requested by cvar or if the benchmark is running (since it can interfere with timings under extreme load)
@@ -3290,7 +3323,7 @@ public:
 				StallCounter--;
 				if (StallCounter == 0 || ElapsedMS < TimeBetweenUpdatesMS)
 				{
-					FPlatformProcess::SetThreadPriority(Priority);
+					FPlatformProcess::SetThreadPriority(CurrentThreadPriority);
 					StallGameThreadCs.Unlock();
 					StallCounter = 0;
 				}
@@ -3361,6 +3394,30 @@ public:
 	}
 
 private:
+	static EThreadPriority GetThreadPriorityFromIndex(int32 Index)
+	{
+		const int32 ClampedIndex = FMath::Clamp(Index, 0, (int32)UE_ARRAY_COUNT(GCsvThreadPriorities) - 1);
+		return GCsvThreadPriorities[ClampedIndex];
+	}
+
+	static const TCHAR* GetThreadPriorityNameFromIndex(int32 Index)
+	{
+		const int32 ClampedIndex = FMath::Clamp(Index, 0, (int32)UE_ARRAY_COUNT(GCsvThreadPriorityNames) - 1);
+		return GCsvThreadPriorityNames[ClampedIndex];
+	}
+	void UpdateThreadPriorityIfNeeded()
+	{
+		EThreadPriority DesiredThreadPriority = GetThreadPriorityFromIndex(GCsvThreadPriorityIndex);
+		if (DesiredThreadPriority != CurrentThreadPriority)
+		{
+			UE_LOG(LogCsvProfiler, Display, TEXT("Setting CSV processing thread priority to %s"), GetThreadPriorityNameFromIndex(GCsvThreadPriorityIndex));
+			FPlatformProcess::SetThreadPriority(DesiredThreadPriority);
+			CurrentThreadPriority = DesiredThreadPriority;
+		}
+	}
+
+
+
 	FRunnableThread* Thread;
 	FCsvProfiler& CsvProfiler;
 	FCriticalSection StallGameThreadCs;
