Convert type traits in Array.h to use `static constexpr bool Value = ...` instead of `enum { Value = ... }`:
C++20 requires the atomic constraints in a requires constraint expression to be of type bool, not just convertible to bool.
#Synced-CL 51218577

diff --git a/Engine/Source/Runtime/Core/Public/Containers/Array.h b/Engine/Source/Runtime/Core/Public/Containers/Array.h
--- a/Engine/Source/Runtime/Core/Public/Containers/Array.h
+++ b/Engine/Source/Runtime/Core/Public/Containers/Array.h
@@ -45,7 +45,7 @@
 template <typename T>
 struct TCanBulkSerialize
 {
-	enum { Value = std::is_arithmetic_v<T> };
+	static constexpr bool Value = std::is_arithmetic_v<T>;
 };
 
 // Forward declarations
@@ -4056,13 +4056,13 @@ DECLARE_TEMPLATE_INTRINSIC_TYPE_LAYOUT((template <typename T, typename Allocator
 template <typename InElementType, typename AllocatorType>
 struct TIsZeroConstructType<TArray<InElementType, AllocatorType>>
 {
-	enum { Value = ::TAllocatorTraits<AllocatorType>::IsZeroConstruct };
+	static constexpr bool Value = ::TAllocatorTraits<AllocatorType>::IsZeroConstruct;
 };
 
 template <typename T, typename AllocatorType>
 struct TIsContiguousContainer<TArray<T, AllocatorType>>
 {
-	enum { Value = true };
+	static constexpr bool Value = true;
 };
 
 /**
@@ -4078,7 +4078,7 @@ template <typename InElementType, typename InAllocatorType> constexpr bool TIsTA
 template <typename T>
 struct TIsTArray
 {
-	enum { Value = TIsTArray_V<T> };
+	static constexpr bool Value = TIsTArray_V<T>;
 };
 
 
