Add the option to output Log to the CrashUserData.

#Synced-CL 51188288

#Synced-CL 51188288

diff --git a/Engine/Source/Runtime/Core/Private/Misc/OutputDeviceMemory.cpp b/Engine/Source/Runtime/Core/Private/Misc/OutputDeviceMemory.cpp
--- a/Engine/Source/Runtime/Core/Private/Misc/OutputDeviceMemory.cpp
+++ b/Engine/Source/Runtime/Core/Private/Misc/OutputDeviceMemory.cpp
@@ -144,3 +144,61 @@ void FOutputDeviceMemory::Dump(FArchive& Ar)
 		Ar.Serialize(PreservePosPtr, AfterPreserveLength * sizeof(ANSICHAR));
 	}
 }
+
+bool FOutputDeviceMemory::DeserializeToBuffer(TArrayView<ANSICHAR> OutBuffer, int32 InOffsetInBytes) const
+{
+	const int32 PreservedSectionSize = FMath::Min(PreserveSize, BufferLength);  
+	const int32 BufferCapacity = Buffer.Num();
+
+	ANSICHAR* DataDest = OutBuffer.GetData();
+	int32 SizeToCopy = OutBuffer.Num();
+	
+	int32 BytesCopied = 0;                                                                                                                                                                                                                                                                                                                                                             
+	int32 BytesToSkip = InOffsetInBytes;
+
+	// copy the startup logs
+	if (BytesToSkip < PreservedSectionSize)
+	{
+		int32 ChunkSizeToCopy = FMath::Min(PreservedSectionSize - BytesToSkip, SizeToCopy);
+		const ANSICHAR* PreservedSectionPosPtr = Buffer.GetData() + BytesToSkip * sizeof(ANSICHAR);
+		FMemory::Memcpy(DataDest, PreservedSectionPosPtr, ChunkSizeToCopy);
+		BytesCopied += ChunkSizeToCopy;
+		DataDest += ChunkSizeToCopy;
+		BytesToSkip = 0;
+	}
+	else
+	{
+		BytesToSkip = FMath::Max(BytesToSkip - PreservedSectionSize, 0);
+	}
+
+	// If the log has wrapped, copy the earliest portion chronologically from the end
+	if (BufferLength == BufferCapacity && BytesCopied < SizeToCopy)
+	{
+		const int32 PreWrapLength = BufferCapacity - BufferStartPos;
+		if (BytesToSkip < PreWrapLength)
+		{
+			int32 ChunkSizeToCopy = FMath::Min(PreWrapLength - BytesToSkip, SizeToCopy - BytesCopied);
+			const ANSICHAR* BufferStartPosPtr = Buffer.GetData() + BufferStartPos + BytesToSkip * sizeof(ANSICHAR);
+			FMemory::Memcpy(DataDest, BufferStartPosPtr, ChunkSizeToCopy);
+			BytesCopied += ChunkSizeToCopy;
+			DataDest += ChunkSizeToCopy;
+			BytesToSkip = 0;
+		}
+		else
+		{
+			BytesToSkip = FMath::Max(BytesToSkip - PreWrapLength, 0);
+		}
+	}
+
+	// Dump the logs from preserved section up to current, if we've made it that far
+	if (BufferLength > PreserveSize && BytesCopied < SizeToCopy && BytesToSkip < BufferStartPos - PreserveSize)
+	{
+		int32 ChunkSizeToCopy = FMath::Min(BufferStartPos - PreserveSize - BytesToSkip, SizeToCopy - BytesCopied);
+		const ANSICHAR* PreservePosPtr = Buffer.GetData() + PreserveSize + BytesToSkip * sizeof(ANSICHAR);
+		FMemory::Memcpy(DataDest, PreservePosPtr, ChunkSizeToCopy);
+		BytesCopied += ChunkSizeToCopy;
+		DataDest += ChunkSizeToCopy;
+	}
+
+	return BytesCopied != 0;
+}
\ No newline at end of file
diff --git a/Engine/Source/Runtime/Core/Public/Misc/OutputDeviceMemory.h b/Engine/Source/Runtime/Core/Public/Misc/OutputDeviceMemory.h
--- a/Engine/Source/Runtime/Core/Public/Misc/OutputDeviceMemory.h
+++ b/Engine/Source/Runtime/Core/Public/Misc/OutputDeviceMemory.h
@@ -28,6 +28,8 @@ public:
 
 	/** Dumps the contents of the buffer to an archive */
 	CORE_API virtual void Dump(FArchive& Ar) override;
+	/** Deserialize the contents to a buffer */
+	CORE_API bool DeserializeToBuffer(TArrayView<ANSICHAR> OutBuffer, int32 InOffsetInBytes = 0) const;
 
 	//~ Begin FOutputDevice Interface.
 	/**
