Fix a crash on Nintendo when running with early csv profiling  (ie -csvCaptureFrames=N). TLS is not initialized when FCsvProfiler::PreInit is called, so we need to defer execution on the processing thread until FCsvProfiler::Init is called.
#tests tested on console
#Synced-CL 51245126

diff --git a/Engine/Source/Runtime/Core/Private/ProfilingDebugging/CsvProfiler.cpp b/Engine/Source/Runtime/Core/Private/ProfilingDebugging/CsvProfiler.cpp
--- a/Engine/Source/Runtime/Core/Private/ProfilingDebugging/CsvProfiler.cpp
+++ b/Engine/Source/Runtime/Core/Private/ProfilingDebugging/CsvProfiler.cpp
@@ -333,6 +333,8 @@ static uint64 GCaptureStartTimestamp = 0;
 
 static bool GCsvProfilerInitialized = false;
 
+static std::atomic<bool> GDelayProcessingThreadStart = false;
+
 //
 // Forward declarations
 //
@@ -3280,6 +3282,11 @@ public:
 
 	virtual uint32 Run() override
 	{
+		// Sleep until the csv profiler is fully initialized. On some platforms, TLS is not ready until this point
+		while (GDelayProcessingThreadStart)
+		{
+			FPlatformProcess::SleepNoStats(0.001f);
+		}
 		GCsvProcessingThreadId = FPlatformTLS::GetCurrentThreadId();
 		GGameThreadIsCsvProcessingThread = false;
 
@@ -4896,12 +4903,21 @@ void FCsvProfiler::PreInit(const TCHAR * Commandline)
 		{
 			FCommandLine::Set(Commandline);
 		}
-		FCsvProfiler::Get()->Init();
+		// Delay processing thread start
+		GDelayProcessingThreadStart = true;
+		FCsvProfiler::Get()->InitInternal();
 	}
 #endif
 }
 
 void FCsvProfiler::Init()
+{
+	InitInternal();
+	// It's now safe to start the processing thread
+	GDelayProcessingThreadStart = false;
+}
+
+void FCsvProfiler::InitInternal()
 {
 	if (GCsvProfilerInitialized)
 	{
@@ -5405,3 +5421,4 @@ void CSVTest()
 #endif // CSV_PROFILER_ALLOW_DEBUG_FEATURES
 
 #endif // CSV_PROFILER
+//UE_ENABLE_OPTIMIZATION
\ No newline at end of file
diff --git a/Engine/Source/Runtime/Core/Public/ProfilingDebugging/CsvProfiler.h b/Engine/Source/Runtime/Core/Public/ProfilingDebugging/CsvProfiler.h
--- a/Engine/Source/Runtime/Core/Public/ProfilingDebugging/CsvProfiler.h
+++ b/Engine/Source/Runtime/Core/Public/ProfilingDebugging/CsvProfiler.h
@@ -580,6 +580,8 @@ public:
 
 private:
 	CORE_API static void VARARGS RecordEventfInternal(int32 CategoryIndex, bool bLog, const TCHAR* Fmt, ...);
+
+	void InitInternal();
 	
 	static void RecordRegionEventInternal(const TCHAR* CategoryName, const TCHAR* RegionName, const TCHAR* SubCategoryName, bool bBegin);
 
