Rename csv.threadPriority -> csv.ProcessingThread.Priority for consistency with other cvars
#Synced-CL 51245126
#Synced-CL 51245126

diff --git a/Engine/Source/Runtime/Core/Private/ProfilingDebugging/CsvProfiler.cpp b/Engine/Source/Runtime/Core/Private/ProfilingDebugging/CsvProfiler.cpp
--- a/Engine/Source/Runtime/Core/Private/ProfilingDebugging/CsvProfiler.cpp
+++ b/Engine/Source/Runtime/Core/Private/ProfilingDebugging/CsvProfiler.cpp
@@ -285,7 +285,7 @@ const TCHAR* GCsvThreadPriorityNames[] =
 
 static int32 GCsvThreadPriorityIndex = 0; 
 FAutoConsoleVariableRef CVarCsvThreadPriority(
-	TEXT("csv.ThreadPriority"),
+	TEXT("csv.ProcessingThread.Priority"),
 	GCsvThreadPriorityIndex,
 	TEXT("Thread priority of the CSV profiler thread: 0=Lowest, 1=BelowNormal, 2=SlightlyBelowNormal, 3=Normal, 4=AboveNormal, 5=Highest\n")
 	TEXT("This is switchable at runtime"),
