Fixed UE::CScopedEnum to be in terms of concepts instead of traits.

diff --git a/Engine/Source/Runtime/Core/Public/Concepts/ScopedEnum.h b/Engine/Source/Runtime/Core/Public/Concepts/ScopedEnum.h
--- a/Engine/Source/Runtime/Core/Public/Concepts/ScopedEnum.h
+++ b/Engine/Source/Runtime/Core/Public/Concepts/ScopedEnum.h
@@ -2,7 +2,8 @@
 
 #pragma once
 
-#include <type_traits>
+#include "Concepts/ConvertibleTo.h"
+#include "Concepts/Enum.h"
 
 namespace UE
 {
@@ -10,5 +11,5 @@ namespace UE
 	 * Concept which describes a scoped enum type (aka. enum class).
 	 */
 	template <typename T>
-	concept CScopedEnum = std::is_enum_v<T> && !std::is_convertible_v<T, int>;
+	concept CScopedEnum = UE::CEnum<T> && !CConvertibleTo<T, int>;
 }
