Removed separate vm blocks per MB3 bin support from the allocator

#Synced-CL 51337004

diff --git a/Engine/Source/Runtime/Core/Private/HAL/MallocBinned3.cpp b/Engine/Source/Runtime/Core/Private/HAL/MallocBinned3.cpp
--- a/Engine/Source/Runtime/Core/Private/HAL/MallocBinned3.cpp
+++ b/Engine/Source/Runtime/Core/Private/HAL/MallocBinned3.cpp
@@ -32,8 +32,6 @@ PRAGMA_DISABLE_UNSAFE_TYPECAST_WARNINGS
 	TAtomic<int64> Binned3Decommits;
 
 	int64 Binned3FreeBitsMemory = 0;
-	TAtomic<int64> Binned3TotalPoolSearches;
-	TAtomic<int64> Binned3TotalPointerTests;
 #endif
 
 #define UE_MB3_TIME_LARGE_BLOCKS (0)
@@ -47,15 +45,7 @@ PRAGMA_DISABLE_UNSAFE_TYPECAST_WARNINGS
 #endif
 
 uint16 FMallocBinned3::SmallBinSizesShifted[UE_MB3_SMALL_POOL_COUNT + 1] = { 0 };
-
-#if !BINNED3_USE_SEPARATE_VM_PER_POOL
-	uint8* FMallocBinned3::Binned3BaseVMPtr = nullptr;
-#else
-	uint64 FMallocBinned3::PoolSearchDiv = 0;
-	uint8* FMallocBinned3::HighestPoolBaseVMPtr = nullptr;
-	uint8* FMallocBinned3::PoolBaseVMPtr[UE_MB3_SMALL_POOL_COUNT] = { nullptr };
-#endif
-
+uint8* FMallocBinned3::Binned3BaseVMPtr = nullptr;
 FMallocBinned3* FMallocBinned3::MallocBinned3 = nullptr;
 // Mapping of sizes to small table indices
 uint8 FMallocBinned3::MemSizeToPoolIndex[SIZE_TO_POOL_INDEX_NUM] = { 0 };
@@ -449,80 +439,17 @@ FMallocBinned3::FMallocBinned3(const FPlatformMemory::FPlatformVirtualMemoryBloc
 
 	// First thing we try to allocate address space for bins as it might help us to move forward Constants.AddressStart and reduce the amount of available address space for the Large OS Allocs
 	// Available address space is used to reserve hash map that can address all of that range, so less addressable space means less memory is allocated for book keeping
-#if !BINNED3_USE_SEPARATE_VM_PER_POOL
 	Binned3BaseVMBlock = ExternalMemoryBlock ? *ExternalMemoryBlock : FPlatformMemory::FPlatformVirtualMemoryBlock::AllocateVirtual(UE_MB3_SMALL_POOL_COUNT * UE_MB3_MAX_MEMORY_PER_POOL_SIZE, OsAllocationGranularity);
 	check(Binned3BaseVMBlock.GetActualSize() >= UE_MB3_SMALL_POOL_COUNT * UE_MB3_MAX_MEMORY_PER_POOL_SIZE);
 	Binned3BaseVMPtr = (uint8*)Binned3BaseVMBlock.GetVirtualPointer();
 	check(IsAligned(Binned3BaseVMPtr, OsAllocationGranularity));
 	verify(Binned3BaseVMPtr);
-#else
-
-	for (uint32 Index = 0; Index < UE_MB3_SMALL_POOL_COUNT; ++Index)
-	{
-		FPlatformMemory::FPlatformVirtualMemoryBlock NewBLock = FPlatformMemory::FPlatformVirtualMemoryBlock::AllocateVirtual(UE_MB3_MAX_MEMORY_PER_POOL_SIZE, OsAllocationGranularity);
-
-		uint8* NewVM = (uint8*)NewBLock.GetVirtualPointer();
-		check(IsAligned(NewVM, OsAllocationGranularity));
-		// insertion sort
-		if (Index && NewVM < PoolBaseVMPtr[Index - 1])
-		{
-			uint32 InsertIndex = 0;
-			for (; InsertIndex < Index; ++InsertIndex)
-			{
-				if (NewVM < PoolBaseVMPtr[InsertIndex])
-				{
-					break;
-				}
-			}
-			check(InsertIndex < Index);
-			for (uint32 MoveIndex = Index; MoveIndex > InsertIndex; --MoveIndex)
-			{
-				PoolBaseVMPtr[MoveIndex] = PoolBaseVMPtr[MoveIndex - 1];
-				PoolBaseVMBlock[MoveIndex] = PoolBaseVMBlock[MoveIndex - 1];
-			}
-			PoolBaseVMPtr[InsertIndex] = NewVM;
-			PoolBaseVMBlock[InsertIndex] = NewBLock;
-		}
-		else
-		{
-			PoolBaseVMPtr[Index] = NewVM;
-			PoolBaseVMBlock[Index] = NewBLock;
-		}
-	}
-	HighestPoolBaseVMPtr = PoolBaseVMPtr[UE_MB3_SMALL_POOL_COUNT - 1];
-	uint64 TotalGaps = 0;
-	for (uint32 Index = 0; Index < UE_MB3_SMALL_POOL_COUNT - 1; ++Index)
-	{
-		check(PoolBaseVMPtr[Index + 1] > PoolBaseVMPtr[Index]); // we sorted it
-		check(PoolBaseVMPtr[Index + 1] >= PoolBaseVMPtr[Index] + UE_MB3_MAX_MEMORY_PER_POOL_SIZE); // and pools are non-overlapping
-		TotalGaps += PoolBaseVMPtr[Index + 1] - (PoolBaseVMPtr[Index] + UE_MB3_MAX_MEMORY_PER_POOL_SIZE);
-	}
-	if (TotalGaps == 0)
-	{
-		PoolSearchDiv = 0;
-	}
-	else if (TotalGaps < UE_MB3_MAX_MEMORY_PER_POOL_SIZE)
-	{
-		PoolSearchDiv = UE_MB3_MAX_MEMORY_PER_POOL_SIZE; // the gaps are not significant, ignoring them should give accurate searches
-	}
-	else
-	{
-		PoolSearchDiv = UE_MB3_MAX_MEMORY_PER_POOL_SIZE + ((TotalGaps + UE_MB3_SMALL_POOL_COUNT - 2) / (UE_MB3_SMALL_POOL_COUNT - 1));
-	}
-#endif
 
 	FGenericPlatformMemoryConstants Constants = FPlatformMemory::GetConstants();
-#if !BINNED3_USE_SEPARATE_VM_PER_POOL
 	if (Constants.AddressStart == (uint64)Binned3BaseVMPtr)
 	{
 		Constants.AddressStart += Align(UE_MB3_SMALL_POOL_COUNT * UE_MB3_MAX_MEMORY_PER_POOL_SIZE, OsAllocationGranularity);
 	}
-#else
-	if (!TotalGaps && Constants.AddressStart == (uint64)PoolBaseVMPtr[0])
-	{
-		Constants.AddressStart += Align(UE_MB3_SMALL_POOL_COUNT * UE_MB3_MAX_MEMORY_PER_POOL_SIZE, OsAllocationGranularity);
-	}
-#endif
 
 	// large slab sizes are possible OsAllocationGranularity = 65536;
 	NumPoolsPerPage = OsAllocationGranularity / sizeof(FPoolInfo);
@@ -616,13 +543,7 @@ FMallocBinned3::~FMallocBinned3()
 void FMallocBinned3::Commit(uint32 InPoolIndex, void *Ptr, SIZE_T Size)
 {
 	UE_MBC_UPDATE_STATS(Binned3Commits++);
-
-#if !BINNED3_USE_SEPARATE_VM_PER_POOL
 	Binned3BaseVMBlock.CommitByPtr(Ptr, Size);
-#else
-	PoolBaseVMBlock[InPoolIndex].CommitByPtr(Ptr, Size);
-#endif
-
 	LLM_IF_ENABLED(FLowLevelMemTracker::Get().OnLowLevelAlloc(ELLMTracker::Platform, Ptr, Size));
 }
 
@@ -630,12 +551,7 @@ void FMallocBinned3::Decommit(uint32 InPoolIndex, void *Ptr, SIZE_T Size)
 {
 	LLM_IF_ENABLED(FLowLevelMemTracker::Get().OnLowLevelFree(ELLMTracker::Platform, Ptr));
 	UE_MBC_UPDATE_STATS(Binned3Decommits++);
-
-#if !BINNED3_USE_SEPARATE_VM_PER_POOL
 	Binned3BaseVMBlock.DecommitByPtr(Ptr, Size);
-#else
-	PoolBaseVMBlock[InPoolIndex].DecommitByPtr(Ptr, Size);
-#endif
 }
 
 void* FMallocBinned3::AllocateMetaDataMemory(SIZE_T Size)
@@ -1023,14 +939,6 @@ void FMallocBinned3::FlushCurrentThreadCacheInternal(bool bNewEpochOnly)
 	FMallocBinnedCommonUtils::FlushCurrentThreadCache(*this, bNewEpochOnly);
 }
 
-#if UE_MB3_ALLOCATOR_STATS && BINNED3_USE_SEPARATE_VM_PER_POOL
-void FMallocBinned3::RecordPoolSearch(uint32 Tests) const
-{
-	Binned3TotalPoolSearches++;
-	Binned3TotalPointerTests += Tests;
-}
-#endif
-
 void FMallocBinned3::DumpAllocatorStats(class FOutputDevice& Ar)
 {
 #if UE_MB3_ALLOCATOR_STATS
@@ -1050,18 +958,6 @@ void FMallocBinned3::DumpAllocatorStats(class FOutputDevice& Ar)
 	Ar.Logf(TEXT("TLS: %fmb"), ((double)TLSMemory.load(std::memory_order_relaxed)) / (1024.0f * 1024.0f));
 	Ar.Logf(TEXT("Slab Commits: %llu"), Binned3Commits.Load());
 	Ar.Logf(TEXT("Slab Decommits: %llu"), Binned3Decommits.Load());
-#if BINNED3_USE_SEPARATE_VM_PER_POOL
-	Ar.Logf(TEXT("BINNED3_USE_SEPARATE_VM_PER_POOL is true - VM is Contiguous = %d"), PoolSearchDiv == 0);
-	if (PoolSearchDiv)
-	{
-		Ar.Logf(TEXT("%llu Pointer Searches   %llu Pointer Compares    %llu Compares/Search"), Binned3TotalPoolSearches.Load(), Binned3TotalPointerTests.Load(), Binned3TotalPointerTests.Load() / Binned3TotalPoolSearches.Load());
-		const uint64 TotalMem = PoolBaseVMPtr[UE_MB3_SMALL_POOL_COUNT - 1] + UE_MB3_MAX_MEMORY_PER_POOL_SIZE - PoolBaseVMPtr[0];
-		const uint64 MinimumMem = uint64(UE_MB3_SMALL_POOL_COUNT) * UE_MB3_MAX_MEMORY_PER_POOL_SIZE;
-		Ar.Logf(TEXT("Percent of gaps in the address range %6.4f  (hopefully < 1, or the searches above will suffer)"), 100.0f * (1.0f - float(MinimumMem) / float(TotalMem)));
-	}
-#else
-	Ar.Logf(TEXT("BINNED3_USE_SEPARATE_VM_PER_POOL is false"));
-#endif
 	Ar.Logf(TEXT("Total allocated from OS: %fmb"), 
 		((double)
 			AllocatedOSSmallPoolMemory.load(std::memory_order_relaxed) + AllocatedLargePoolMemoryWAlignment.load(std::memory_order_relaxed) + PoolInfoMemory + HashMemory + Binned3FreeBitsMemory + TLSMemory.load(std::memory_order_relaxed)
diff --git a/Engine/Source/Runtime/Core/Public/HAL/MallocBinned3.h b/Engine/Source/Runtime/Core/Public/HAL/MallocBinned3.h
--- a/Engine/Source/Runtime/Core/Public/HAL/MallocBinned3.h
+++ b/Engine/Source/Runtime/Core/Public/HAL/MallocBinned3.h
@@ -49,15 +49,6 @@
 
 #define UE_MB3_MAX_MEMORY_PER_POOL_SIZE				(1ull << UE_MB3_MAX_MEMORY_PER_POOL_SIZE_SHIFT) 
 
-// This choice depends on how efficient the OS is with sparse commits in large VM blocks
-#if !defined(BINNED3_USE_SEPARATE_VM_PER_POOL)
-#	if PLATFORM_WINDOWS
-#		define BINNED3_USE_SEPARATE_VM_PER_POOL		(1)
-#	else
-#		define BINNED3_USE_SEPARATE_VM_PER_POOL		(0)
-#	endif
-#endif
-
 #define UE_MB3_ALLOCATOR_STATS						UE_MBC_ALLOCATOR_STATS
 
 #if UE_MB3_ALLOCATOR_STATS
@@ -74,7 +65,6 @@ PRAGMA_DISABLE_UNSAFE_TYPECAST_WARNINGS
 // 
 // MallocBinned3 supports two types of allocations - large and small pool allocation:
 // 1. For small pool allocation MallocBinned3 reserves contiguous range of virtual memory (a Pool) for each allocation size (a Bin).
-//    These Pools can be adjacent to each other in memory if BINNED3_USE_SEPARATE_VM_PER_POOL is set to 0.
 //    By default each Pool reserves 1 GB of address space, unless USE_512MB_MAX_MEMORY_PER_BLOCK_SIZE is defined to 1.
 //    Each Pool commits and decommits reserved memory in Blocks. Each Block is equal to at least one memory page in it's size.
 //    A Block contains N Bins in a way to minimize the tail memory waste. For example 16 bytes bins fit inside 4KB memory page without any waste;
@@ -189,76 +179,15 @@ private:
 	UE::FPlatformRecursiveMutex Mutex;
 #endif
 
-#if !BINNED3_USE_SEPARATE_VM_PER_POOL
 	FORCEINLINE uint64 PoolIndexFromPtr(const void* Ptr) const		// returns a uint64 for it can also be used to check if it is an OS allocation
 	{
 		return (UPTRINT(Ptr) - UPTRINT(Binned3BaseVMPtr)) >> UE_MB3_MAX_MEMORY_PER_POOL_SIZE_SHIFT;
 	}
-	FORCEINLINE uint8* PoolBasePtr(uint32 InPoolIndex) const
-	{
-		return Binned3BaseVMPtr + InPoolIndex * UE_MB3_MAX_MEMORY_PER_POOL_SIZE;
-	}
-#else
-#	if UE_MB3_ALLOCATOR_STATS
-		void RecordPoolSearch(uint32 Tests) const;
-#	else
-		FORCEINLINE void RecordPoolSearch(uint32 Tests) const
-		{
-
-		}
-#	endif
-
-	inline uint64 PoolIndexFromPtr(const void* Ptr) const		// returns a uint64 for it can also be used to check if it is an OS allocation
-	{
-		if (PoolSearchDiv == 0)
-		{
-			return (UPTRINT(Ptr) - UPTRINT(PoolBaseVMPtr[0])) >> UE_MB3_MAX_MEMORY_PER_POOL_SIZE_SHIFT;
-		}
-		uint64 PoolIndex = UE_MB3_SMALL_POOL_COUNT;
-		if (((uint8*)Ptr >= PoolBaseVMPtr[0]) & ((uint8*)Ptr < HighestPoolBaseVMPtr + UE_MB3_MAX_MEMORY_PER_POOL_SIZE))
-		{
-			PoolIndex = uint64((uint8*)Ptr - PoolBaseVMPtr[0]) / PoolSearchDiv;
-			if (PoolIndex >= UE_MB3_SMALL_POOL_COUNT)
-			{
-				PoolIndex = UE_MB3_SMALL_POOL_COUNT - 1;
-			}
-			uint32 Tests = 1; // we are counting potential cache misses here, not actual comparisons
-			if ((uint8*)Ptr < PoolBaseVMPtr[PoolIndex])
-			{
-				do
-				{
-					Tests++;
-					PoolIndex--;
-					check(PoolIndex < UE_MB3_SMALL_POOL_COUNT);
-				} while ((uint8*)Ptr < PoolBaseVMPtr[PoolIndex]);
-				if ((uint8*)Ptr >= PoolBaseVMPtr[PoolIndex] + UE_MB3_MAX_MEMORY_PER_POOL_SIZE)
-				{
-					PoolIndex = UE_MB3_SMALL_POOL_COUNT; // was in the gap
-				}
-			}
-			else if ((uint8*)Ptr >= PoolBaseVMPtr[PoolIndex] + UE_MB3_MAX_MEMORY_PER_POOL_SIZE)
-			{
-				do
-				{
-					Tests++;
-					PoolIndex++;
-					check(PoolIndex < UE_MB3_SMALL_POOL_COUNT);
-				} while ((uint8*)Ptr >= PoolBaseVMPtr[PoolIndex] + UE_MB3_MAX_MEMORY_PER_POOL_SIZE);
-				if ((uint8*)Ptr < PoolBaseVMPtr[PoolIndex])
-				{
-					PoolIndex = UE_MB3_SMALL_POOL_COUNT; // was in the gap
-				}
-			}
-			RecordPoolSearch(Tests);
-		}
-		return PoolIndex;
-	}
 
 	FORCEINLINE uint8* PoolBasePtr(uint32 InPoolIndex) const
 	{
-		return PoolBaseVMPtr[InPoolIndex];
+		return Binned3BaseVMPtr + InPoolIndex * UE_MB3_MAX_MEMORY_PER_POOL_SIZE;
 	}
-#endif //~!BINNED3_USE_SEPARATE_VM_PER_POOL
 
 	inline uint32 PoolIndexFromPtrChecked(const void* Ptr) const
 	{
@@ -358,15 +287,9 @@ public:
 	static uint16 SmallBinSizesShifted[UE_MB3_SMALL_POOL_COUNT + 1];
 	static FMallocBinned3* MallocBinned3;
 
-#if !BINNED3_USE_SEPARATE_VM_PER_POOL
 	static uint8* Binned3BaseVMPtr;
 	FPlatformMemory::FPlatformVirtualMemoryBlock Binned3BaseVMBlock;
-#else
-	static uint64 PoolSearchDiv; // if this is zero, the VM turned out to be contiguous anyway so we use a simple subtract and shift
-	static uint8* HighestPoolBaseVMPtr; // this is a duplicate of PoolBaseVMPtr[UE_MB3_SMALL_POOL_COUNT - 1]
-	static uint8* PoolBaseVMPtr[UE_MB3_SMALL_POOL_COUNT];
-	FPlatformMemory::FPlatformVirtualMemoryBlock PoolBaseVMBlock[UE_MB3_SMALL_POOL_COUNT];
-#endif
+
 	// Mapping of sizes to small table indices
 	static uint8 MemSizeToPoolIndex[SIZE_TO_POOL_INDEX_NUM];
 
diff --git a/Engine/Source/Runtime/Core/Public/HAL/MallocBinnedCommon.h b/Engine/Source/Runtime/Core/Public/HAL/MallocBinnedCommon.h
--- a/Engine/Source/Runtime/Core/Public/HAL/MallocBinnedCommon.h
+++ b/Engine/Source/Runtime/Core/Public/HAL/MallocBinnedCommon.h
@@ -32,14 +32,6 @@
 #	define UE_MBC_NUM_LISTED_SMALL_POOLS				52
 #endif
 
-#if !defined(BINNEDCOMMON_USE_SEPARATE_VM_PER_POOL)
-#	if PLATFORM_WINDOWS
-#		define BINNEDCOMMON_USE_SEPARATE_VM_PER_POOL	1
-#	else
-#		define BINNEDCOMMON_USE_SEPARATE_VM_PER_POOL	0
-#	endif
-#endif
-
 #define UE_MBC_MIN_SMALL_POOL_ALIGNMENT					8	// Minimum alignment of bins. Added to support 8 bytes bin. If you need to mask lower bits - use this!
 #define UE_MBC_MAX_SMALL_POOL_ALIGNMENT					256
 #define UE_MBC_STANDARD_ALIGNMENT						16	// Standard alignment for all allocations, except 8 bytes that is aligned by 8 bytes
diff --git a/Engine/Source/Runtime/Core/Public/HAL/MallocBinnedGPU.h b/Engine/Source/Runtime/Core/Public/HAL/MallocBinnedGPU.h
--- a/Engine/Source/Runtime/Core/Public/HAL/MallocBinnedGPU.h
+++ b/Engine/Source/Runtime/Core/Public/HAL/MallocBinnedGPU.h
@@ -55,7 +55,7 @@ struct FArenaParams
 	uint8 MaxGlobalBundles = 32;
 	uint8 MinimumAlignmentShift = 4;
 	uint8 PoolCount;
-	bool bUseSeparateVMPerPool = !!(BINNEDCOMMON_USE_SEPARATE_VM_PER_POOL);
+	bool bUseSeparateVMPerPool = false;
 	bool bPerThreadCaches = true;
 	bool bUseStandardSmallPoolSizes = true;
 	bool bAttemptToAlignSmallBocks = true;
