Added EnumToUnderlyingType, for converting an enumerator value without needing to specify the underlying type.

diff --git a/Engine/Source/Runtime/Core/Public/Misc/EnumClassFlags.h b/Engine/Source/Runtime/Core/Public/Misc/EnumClassFlags.h
--- a/Engine/Source/Runtime/Core/Public/Misc/EnumClassFlags.h
+++ b/Engine/Source/Runtime/Core/Public/Misc/EnumClassFlags.h
@@ -41,18 +41,18 @@ UE_REWRITE constexpr bool EnumHasAnyFlags(Enum Flags, Enum Contains)
 }
 
 /**
-* Check if the Flags value contains only the flags specified by the Contains argument.
-* It is not required that any of the flags in Contains be set in Flags and Flags may be
-* any combination of the flags specified by Contains (including k=0). But if Flags 
-* contains any flag not in Contains, this function will return false.
-* This function may also be understood as !EnumContainsAnyFlags(Flags, ~Contains).
-*
-* @param Flags    The value to check.
-* @param Contains The flags to check the value against.
-*
-* @return true if Flags consists of only the flags specified by Contains. false otherwise.
-* @note   Returns true if Flags is empty.
-*/
+ * Check if the Flags value contains only the flags specified by the Contains argument.
+ * It is not required that any of the flags in Contains be set in Flags and Flags may be
+ * any combination of the flags specified by Contains (including k=0). But if Flags 
+ * contains any flag not in Contains, this function will return false.
+ * This function may also be understood as !EnumContainsAnyFlags(Flags, ~Contains).
+ *
+ * @param Flags    The value to check.
+ * @param Contains The flags to check the value against.
+ *
+ * @return true if Flags consists of only the flags specified by Contains. false otherwise.
+ * @note   Returns true if Flags is empty.
+ */
 template <typename Enum>
 UE_REWRITE constexpr bool EnumOnlyContainsFlags(Enum Flags, Enum Contains)
 {
@@ -60,13 +60,13 @@ UE_REWRITE constexpr bool EnumOnlyContainsFlags(Enum Flags, Enum Contains)
 }
 
 /**
-* Check if Flags has one and only one flag set.
-* This can also be thought of as a check for a power-of-2 value.
-*
-* @param Flags The value to check.
-*
-* @return true if Flags has only one flag set. false otherwise.
-*/
+ * Check if Flags has one and only one flag set.
+ * This can also be thought of as a check for a power-of-2 value.
+ *
+ * @param Flags The value to check.
+ *
+ * @return true if Flags has only one flag set. false otherwise.
+ */
 template <typename Enum>
 UE_REWRITE constexpr bool EnumHasOneFlag(Enum Flags)
 {
@@ -75,13 +75,13 @@ UE_REWRITE constexpr bool EnumHasOneFlag(Enum Flags)
 }
 
 /**
-* Check if Flags has one and only one of the flags specified in OneOfFlags set
-*
-* @param Flags       The value to check.
-* @param OneOfFlags  The flags to check the value for.
-*
-* @return true if Flags has one and only one of the flags in OneOfFlags set. false otherwise.
-*/
+ * Check if Flags has one and only one of the flags specified in OneOfFlags set
+ *
+ * @param Flags       The value to check.
+ * @param OneOfFlags  The flags to check the value for.
+ *
+ * @return true if Flags has one and only one of the flags in OneOfFlags set. false otherwise.
+ */
 template <typename Enum>
 UE_REWRITE constexpr bool EnumHasAnyOneFlag(Enum Flags, Enum OneOfFlags)
 {
@@ -138,3 +138,10 @@ UE_NODEBUG constexpr int EnumNumSetFlags(Enum Flags)
 	}
 	return Result;
 }
+
+// Converts an enumerator value to the same value in its underlying type.
+template <typename Enum>
+UE_NODEBUG constexpr int EnumToUnderlyingType(Enum E)
+{
+	return (__underlying_type(Enum))E;
+}
