Reduce time under the shared large alloc lock in MB3 and unify mutex for cached allocator and shared large allocator

#Synced-CL 51337004
#Synced-CL 51337004

diff --git a/Engine/Source/Runtime/Core/Private/HAL/MallocBinned3.cpp b/Engine/Source/Runtime/Core/Private/HAL/MallocBinned3.cpp
--- a/Engine/Source/Runtime/Core/Private/HAL/MallocBinned3.cpp
+++ b/Engine/Source/Runtime/Core/Private/HAL/MallocBinned3.cpp
@@ -591,7 +591,7 @@ void* FMallocBinned3::Malloc(SIZE_T Size, uint32 Alignment)
 	bool UsePools = (Size <= UE_MB3_MAX_SMALL_POOL_SIZE);
 
 	check(FMath::IsPowerOfTwo(Alignment));
-	
+
 	if (UNLIKELY(UsePools && (Alignment > UE_MBC_STANDARD_ALIGNMENT)))
 	{
 		// check if allocations that require alignment larger than UE_MBC_STANDARD_ALIGNMENT can be promoted to a bin with a natural alignment that matches
@@ -600,7 +600,7 @@ void* FMallocBinned3::Malloc(SIZE_T Size, uint32 Alignment)
 		UsePools = PromoteToLargerBin(Size, Alignment, *this);
 	}
 
-	if (UsePools) 
+	if (UsePools)
 	{
 		const uint32 PoolIndex = BoundSizeToPoolIndex(Size, MemSizeToPoolIndex);
 		FPerThreadFreeBlockLists* Lists = GMallocBinnedPerThreadCaches ? FPerThreadFreeBlockLists::Get() : nullptr;
@@ -632,7 +632,7 @@ void* FMallocBinned3::Malloc(SIZE_T Size, uint32 Alignment)
 		if (!Pool)
 		{
 			Pool = PushNewPoolToFront(Table, Table.BinSize, PoolIndex, BlockIndex);
-			
+
 			// Indicates that we run out of reserved virtual memory for the pool (UE_MB3_MAX_MEMORY_PER_POOL_SIZE) for this bin type
 			if (!Pool)
 			{
@@ -688,26 +688,42 @@ void* FMallocBinned3::Malloc(SIZE_T Size, uint32 Alignment)
 	// Use OS for non-pooled allocations.
 	const uint64 AlignedSize = Align(Size, FPlatformMemory::FPlatformVirtualMemoryBlock::GetCommitAlignment());
 
-#if UE_MB3_TIME_LARGE_BLOCKS
-	const double StartTime = FPlatformTime::Seconds();
-#endif
-
 	LLM_PLATFORM_SCOPE(ELLMTag::FMalloc);
 
 	NOALLOC_SCOPE_CYCLE_COUNTER(STAT_FMallocBinned3_MallocExternalLarge);
 
 	LogLargeAllocation(AlignedSize);
 
-#if UE_MB3_USE_CACHED_PAGE_ALLOCATOR_FOR_LARGE_ALLOCS
-	UE::TUniqueLock Lock(Mutex);
-	void* Result = GetCachedOSPageAllocator().Allocate(AlignedSize);
-	check(IsAligned(Result, Alignment));
-#else
+#if UE_MB3_TIME_LARGE_BLOCKS
+	const double StartTime = FPlatformTime::Seconds();
+#endif
+
+	void* Result;
+#if !UE_MB3_USE_CACHED_PAGE_ALLOCATOR_FOR_LARGE_ALLOCS
 	FPlatformMemory::FPlatformVirtualMemoryBlock Block = FPlatformMemory::FPlatformVirtualMemoryBlock::AllocateVirtual(AlignedSize, Alignment);
 	Block.Commit(0, AlignedSize);
-	void* Result = Block.GetVirtualPointer();
+	Result = Block.GetVirtualPointer();
 	LLM_IF_ENABLED(FLowLevelMemTracker::Get().OnLowLevelAlloc(ELLMTracker::Platform, Result, AlignedSize));
+	const size_t ActualSize = Block.GetActualSizeInPages();
+#endif
+
+	{
+		UE::TUniqueLock Lock(ExternalAllocMutex);
+
+#if UE_MB3_USE_CACHED_PAGE_ALLOCATOR_FOR_LARGE_ALLOCS
+		Result = GetCachedOSPageAllocator().Allocate(AlignedSize);
+		const size_t ActualSize = AlignedSize / FPlatformMemory::FPlatformVirtualMemoryBlock::GetCommitAlignment();
 #endif
+		if (!Result)
+		{
+			OutOfMemory(AlignedSize);
+		}
+
+		// Create pool.
+		FPoolInfo* Pool = Internal::GetOrCreatePoolInfo(*this, Result, FPoolInfo::ECanary::Assigned);
+		check(Size > 0 && Size <= AlignedSize && AlignedSize >= FPlatformMemory::FPlatformVirtualMemoryBlock::GetCommitAlignment());
+		Pool->SetOSAllocationSizes(Size, AlignedSize, ActualSize);
+	}
 
 #if UE_MB3_TIME_LARGE_BLOCKS
 	const double Add = FPlatformTime::Seconds() - StartTime;
@@ -720,31 +736,13 @@ void* FMallocBinned3::Malloc(SIZE_T Size, uint32 Alignment)
 #endif
 
 	UE_CLOG(!IsAligned(Result, Alignment) ,LogMemory, Fatal, TEXT("FMallocBinned3 alignment was too large for OS. Alignment=%d Ptr=%p"), Alignment, Result);
-
-	if (!Result)
-	{
-		OutOfMemory(AlignedSize);
-	}
 	check(IsOSAllocation(Result));
 
-#if! UE_MB3_USE_CACHED_PAGE_ALLOCATOR_FOR_LARGE_ALLOCS
-	UE::TUniqueLock Lock(ExternalAllocMutex);
-#endif
-
 #if UE_MB3_ALLOCATOR_STATS
 	AllocatedLargePoolMemory.fetch_add(Size, std::memory_order_relaxed);
 	AllocatedLargePoolMemoryWAlignment.fetch_add(AlignedSize, std::memory_order_relaxed);
 #endif
 
-	// Create pool.
-	FPoolInfo* Pool = Internal::GetOrCreatePoolInfo(*this, Result, FPoolInfo::ECanary::Assigned);
-	check(Size > 0 && Size <= AlignedSize && AlignedSize >= FPlatformMemory::FPlatformVirtualMemoryBlock::GetCommitAlignment());
-#if UE_MB3_USE_CACHED_PAGE_ALLOCATOR_FOR_LARGE_ALLOCS
-	Pool->SetOSAllocationSizes(Size, AlignedSize, AlignedSize / FPlatformMemory::FPlatformVirtualMemoryBlock::GetCommitAlignment());
-#else
-	Pool->SetOSAllocationSizes(Size, AlignedSize, Block.GetActualSizeInPages());
-#endif
-
 	return Result;
 }
 
@@ -862,41 +860,39 @@ void FMallocBinned3::FreeExternal(void* Ptr, uint64 PoolIndex)
 	}
 	else if (Ptr)
 	{
-#if UE_MB3_USE_CACHED_PAGE_ALLOCATOR_FOR_LARGE_ALLOCS
-		UE::TUniqueLock Lock(Mutex);
-#endif
-		uint32 VMPages;
+		FPoolInfo* Pool;
 		{
-#if !UE_MB3_USE_CACHED_PAGE_ALLOCATOR_FOR_LARGE_ALLOCS
 			UE::TUniqueLock Lock(ExternalAllocMutex);
-#endif
-			FPoolInfo* Pool = Internal::FindPoolInfo(*this, Ptr);
-			if (!Pool)
-			{
-				UE_LOG(LogMemory, Fatal, TEXT("FMallocBinned3 Attempt to free an unrecognized pointer %p"), Ptr);
-			}
-			const uint32 PoolOsBytes = Pool->GetOsCommittedBytes();
-			const uint32 PoolOSRequestedBytes = Pool->GetOSRequestedBytes();
-			VMPages = Pool->GetOsVMPages();
+
+			Pool = Internal::FindPoolInfo(*this, Ptr);
+		}
+
+		if (!Pool)
+		{
+			UE_LOG(LogMemory, Fatal, TEXT("FMallocBinned3 Attempt to free an unrecognized pointer %p"), Ptr);
+		}
+		const uint32 PoolOsBytes = Pool->GetOsCommittedBytes();
+		const uint32 PoolOSRequestedBytes = Pool->GetOSRequestedBytes();
+		const uint32 VMPages = Pool->GetOsVMPages();
 
 #if UE_MB3_ALLOCATOR_STATS
-			AllocatedLargePoolMemory.fetch_sub((int64)PoolOSRequestedBytes, std::memory_order_relaxed);
-			AllocatedLargePoolMemoryWAlignment.fetch_sub((int64)PoolOsBytes, std::memory_order_relaxed);
+		AllocatedLargePoolMemory.fetch_sub((int64)PoolOSRequestedBytes, std::memory_order_relaxed);
+		AllocatedLargePoolMemoryWAlignment.fetch_sub((int64)PoolOsBytes, std::memory_order_relaxed);
 #endif
 
-			checkf(PoolOSRequestedBytes <= PoolOsBytes, TEXT("FMallocBinned3::FreeExternal %u %u"), int32(PoolOSRequestedBytes), int32(PoolOsBytes));
-			Pool->SetCanary(FPoolInfo::ECanary::Unassigned, true, false);
-		}
+		checkf(PoolOSRequestedBytes <= PoolOsBytes, TEXT("FMallocBinned3::FreeExternal %u %u"), int32(PoolOSRequestedBytes), int32(PoolOsBytes));
+		Pool->SetCanary(FPoolInfo::ECanary::Unassigned, true, false);
 
 		// Free an OS allocation.
 #if UE_MB3_TIME_LARGE_BLOCKS
 		const double StartTime = FPlatformTime::Seconds();
 #endif
 		{
-			LLM_IF_ENABLED(FLowLevelMemTracker::Get().OnLowLevelFree(ELLMTracker::Platform, Ptr));
 #if UE_MB3_USE_CACHED_PAGE_ALLOCATOR_FOR_LARGE_ALLOCS
+			UE::TUniqueLock Lock(ExternalAllocMutex);
 			GetCachedOSPageAllocator().Free(Ptr, VMPages * FPlatformMemory::FPlatformVirtualMemoryBlock::GetCommitAlignment());
 #else
+			LLM_IF_ENABLED(FLowLevelMemTracker::Get().OnLowLevelFree(ELLMTracker::Platform, Ptr));
 			FPlatformMemory::FPlatformVirtualMemoryBlock Block(Ptr, VMPages);
 			Block.FreeVirtual();
 #endif
diff --git a/Engine/Source/Runtime/Core/Public/HAL/MallocBinned3.h b/Engine/Source/Runtime/Core/Public/HAL/MallocBinned3.h
--- a/Engine/Source/Runtime/Core/Public/HAL/MallocBinned3.h
+++ b/Engine/Source/Runtime/Core/Public/HAL/MallocBinned3.h
@@ -175,10 +175,6 @@ public:
 private:
 	uint32 SmallPoolInfosPerPlatformPage;
 
-#if	UE_MB3_USE_CACHED_PAGE_ALLOCATOR_FOR_LARGE_ALLOCS
-	UE::FPlatformRecursiveMutex Mutex;
-#endif
-
 	FORCEINLINE uint64 PoolIndexFromPtr(const void* Ptr) const		// returns a uint64 for it can also be used to check if it is an OS allocation
 	{
 		return (UPTRINT(Ptr) - UPTRINT(Binned3BaseVMPtr)) >> UE_MB3_MAX_MEMORY_PER_POOL_SIZE_SHIFT;
