Fixed return type of EnumToUnderlyingType.

diff --git a/Engine/Source/Runtime/Core/Public/Misc/EnumClassFlags.h b/Engine/Source/Runtime/Core/Public/Misc/EnumClassFlags.h
--- a/Engine/Source/Runtime/Core/Public/Misc/EnumClassFlags.h
+++ b/Engine/Source/Runtime/Core/Public/Misc/EnumClassFlags.h
@@ -141,7 +141,7 @@ UE_NODEBUG constexpr int EnumNumSetFlags(Enum Flags)
 
 // Converts an enumerator value to the same value in its underlying type.
 template <typename Enum>
-UE_NODEBUG constexpr int EnumToUnderlyingType(Enum E)
+UE_NODEBUG constexpr auto EnumToUnderlyingType(Enum E)
 {
 	return (__underlying_type(Enum))E;
 }
