CSV profiler: Add additional stats and controls to limit processing thread update latency in order to constrain memory usage
- New stats CsvProfiler/LastUpdateMemUsedMB and CsvProfiler/TimeSinceLastUpdateMs track last update's memory use and current latency
- The cvar csv.ProcessingThread.GtStallMaxUpdateLatencyMs can be used to specify a maximum latency. If time since last update goes beyond this, the gamethread will stall.
- In the event of a GT stall:
- Metadata HasGtStalls will be set to 1 (this defaults to 0)
- A "GT Stall" event will be emitted
- Timing stats CsvProfiler/GameThread/Stall and Exclusive/GameThread/CsvProfiler_Stall will be set
- Processing Thread priority will be boosted (via the bThreadPriorityBoost attribute)
- Also refactor previous logic for handling GT stalls on the processing thread (to catch ever-increasing update scenarios). This now lives in its own function
#tests tested on console
#Synced-CL 51245126

diff --git a/Engine/Source/Runtime/Core/Private/ProfilingDebugging/CsvProfiler.cpp b/Engine/Source/Runtime/Core/Private/ProfilingDebugging/CsvProfiler.cpp
--- a/Engine/Source/Runtime/Core/Private/ProfilingDebugging/CsvProfiler.cpp
+++ b/Engine/Source/Runtime/Core/Private/ProfilingDebugging/CsvProfiler.cpp
@@ -220,10 +220,17 @@ TAutoConsoleVariable<float> CVarCsvProcessingThreadTimeBetweenUpdates(
 	ECVF_Default
 );
 
+TAutoConsoleVariable<float> CVarCsvProcessingThreadGtStallMaxUpdateLatencyMs(
+	TEXT("csv.ProcessingThread.GtStallMaxUpdateLatencyMs"),
+	0.0,
+	TEXT("How far behind the processing thread can get (ms since last update) before we stall the GT. Detection is handled on the GT. Set to 0 to disable"),
+	ECVF_Default
+);
+
 TAutoConsoleVariable<float> CVarCsvProcessingThreadGtStallUpdateTimeThresholdMs(
 	TEXT("csv.ProcessingThread.GtStallUpdateTimeThresholdMs"),
 	1000.0f,
-	TEXT("Specifies the max time a processing thread update can take before we consider stalling the game thread.\n")
+	TEXT("Specifies the max time a processing thread update can take before we consider stalling the game thread. Used in combination with csv.ProcessingThread.GtStallUpdateEscalationThreshold\n")
 	TEXT("Set to 0 to disable stalling"),
 	ECVF_Default
 );
@@ -2339,18 +2346,11 @@ struct FCsvAggregateStatSeries : public FCsvStatSeries
 };
 
 
-struct FCsvProcessThreadDataStats
+uint64 FCsvProcessThreadDataStats::GetEstimatedMemoryUse() const
 {
-	FCsvProcessThreadDataStats()
-		: TimestampCount(0)
-		, CustomStatCount(0)
-		, EventCount(0)
-	{}
-
-	uint32 TimestampCount;
-	uint32 CustomStatCount;
-	uint32 EventCount;
-};
+	// Note: this doesn't include memory for event strings
+	return sizeof(FCsvTimingMarker)*TimestampCount + sizeof(FCsvCustomStat)*CustomStatCount + sizeof(FCsvEvent)*EventCount;
+}
 
 class FCsvThreadGroupStatProcessor;
 
@@ -3233,6 +3233,13 @@ inline bool FCsvStreamWriter::HasReachedMaxFileBytes() const
 }
 
 
+struct FProcessingThreadStats
+{
+	double LastUpdateStartTime = 0.0;
+	double LastUpdateEndTime = 0.0;
+	FCsvProcessThreadDataStats ProcessingStats;
+};
+
 //-----------------------------------------------------------------------------
 //	FCsvProfilerProcessingThread class : low priority thread to process 
 //  profiling data
@@ -3242,6 +3249,8 @@ class FCsvProfilerProcessingThread : public FRunnable
 	FThreadSafeCounter StopCounter;
 	EThreadPriority CurrentThreadPriority;
 
+	FProcessingThreadStats UpdateStats;
+	FCriticalSection StatsCs;
 public:
 	FCsvProfilerProcessingThread(FCsvProfiler& InCsvProfiler)
 		: CsvProfiler(InCsvProfiler)
@@ -3287,20 +3296,15 @@ public:
 		{
 			FPlatformProcess::SleepNoStats(0.001f);
 		}
+		// Set initial LastUpdateTime to the time the thread started so we can compute processing latency on capture start
+		UpdateStats.LastUpdateStartTime = UpdateStats.LastUpdateEndTime = FPlatformTime::Seconds();
+
 		GCsvProcessingThreadId = FPlatformTLS::GetCurrentThreadId();
 		GGameThreadIsCsvProcessingThread = false;
 
 		FMemory::SetupTLSCachesOnCurrentThread();
 
 		LLM_SCOPE(ELLMTag::CsvProfiler);
-
-		double TotalProcessingTime = 0.0;
-		uint64 TotalStatEntriesProcessed = 0;
-		float PrevElapsedMs = 0.0f;
-		uint32 PrevNumStatEntriesProcessed = 0;
-		uint32 StallCounter = 0;
-		int32 SlowUpdateEscalationCount = 0;
-
 		while (StopCounter.GetValue() == 0)
 		{
 			UpdateThreadPriorityIfNeeded();
@@ -3312,11 +3316,12 @@ public:
 				FPlatformProcess::Sleep(5.0f / 1000.0f);
 				continue;
 			}
-			uint32 NumStatEntriesProcessed = 0;
-			float ElapsedMS = CsvProfiler.ProcessStatData(&NumStatEntriesProcessed);
+			double UpdateStartTime = FPlatformTime::Seconds();
 
-			TotalStatEntriesProcessed += (uint64)NumStatEntriesProcessed;
-			TotalProcessingTime += (double)ElapsedMS/1000.0;
+			FCsvProcessThreadDataStats Stats;
+			float ElapsedMs = CsvProfiler.ProcessStatData(&Stats);
+
+			double UpdateEndTime = FPlatformTime::Seconds();
 
 			if (GCsvProfilerIsWritingFile)
 			{
@@ -3324,54 +3329,21 @@ public:
 				CsvProfiler.FileWriteBlockingEvent->Trigger();
 			}
 
-			// If we've stalled the GT then check if we can resume
-			if (StallCounter > 0)
-			{
-				StallCounter--;
-				if (StallCounter == 0 || ElapsedMS < TimeBetweenUpdatesMS)
-				{
-					FPlatformProcess::SetThreadPriority(CurrentThreadPriority);
-					StallGameThreadCs.Unlock();
-					StallCounter = 0;
-				}
-			}
-			else
-			{
-				// Detect well-of-despair (stats coming in faster than they can be processed)
-				const float StallThresholdMs = CVarCsvProcessingThreadGtStallUpdateTimeThresholdMs.GetValueOnAnyThread();
+			HandleGtStalling(Stats, ElapsedMs);
 
-				// If this thread is slower than the last and processed more stats, increase SlowUpdateEscalationCount
-				if (StallThresholdMs > 0.0f && PrevElapsedMs > StallThresholdMs && ElapsedMS > PrevElapsedMs && NumStatEntriesProcessed > PrevNumStatEntriesProcessed)
-				{
-					SlowUpdateEscalationCount++;
-					if (SlowUpdateEscalationCount >= CVarCsvProcessingThreadGtStallUpdateEscalationThreshold.GetValueOnAnyThread())
-					{
-						// Stall the GT for up to 2 updates to allow us to catch up
-						StallGameThreadCs.Lock();
-						StallCounter = 2;
-						double StatEntriesProcessedPerSecond = (double)TotalStatEntriesProcessed / TotalProcessingTime;
-						UE_LOG(LogCsvProfiler, Warning, TEXT("Stats coming in faster than we can process them! GT stalled until we can catch up!"));
-						UE_LOG(LogCsvProfiler, Warning, TEXT("Avg processing rate: %.0f stat entries per second (timestamps+custom stats)"), StatEntriesProcessedPerSecond);
-						UE_LOG(LogCsvProfiler, Warning, TEXT("Check CsvProfiler/* stats to see current rates. Run with -csvStatCounts to report per-stat counts"));
-						// Temporarily boost thread priority since the GT is waiting
-						FPlatformProcess::SetThreadPriority(TPri_AboveNormal);
-						SlowUpdateEscalationCount = 0;
-					}
-				}
-				else
-				{
-					SlowUpdateEscalationCount = 0;
-				}
+			// Update stats
+			{
+				UE::TScopeLock Lock(StatsCs);
+				UpdateStats.LastUpdateStartTime = UpdateStartTime;
+				UpdateStats.LastUpdateEndTime = UpdateEndTime;
+				UpdateStats.ProcessingStats = Stats;
 			}
-			
-			PrevNumStatEntriesProcessed = NumStatEntriesProcessed;
-			PrevElapsedMs = ElapsedMS;
 
-			float SleepTimeSeconds = FMath::Max(TimeBetweenUpdatesMS - ElapsedMS, 0.0f) / 1000.0f;
+			float SleepTimeSeconds = FMath::Max(TimeBetweenUpdatesMS - ElapsedMs, 0.0f) / 1000.0f;
 			FPlatformProcess::Sleep(SleepTimeSeconds);
 		}
 		
-		if (StallCounter>0)
+		if (bHasGtStallLock)
 		{
 			StallGameThreadCs.Unlock();
 		}
@@ -3388,16 +3360,51 @@ public:
 
 	virtual void Exit() override { }
 
-	void StallGameThreadIfNeeded()
+	bool StallGameThreadIfNeeded()
 	{
 		check(IsInGameThread());
+
+		// Handle stalling GT due to max update latency (too long since last update)
+		const double MaxUpdateLatency = (double)CVarCsvProcessingThreadGtStallMaxUpdateLatencyMs.GetValueOnGameThread()/1000.0;
+		if (MaxUpdateLatency > 0.0)
+		{
+			double GtTime = FPlatformTime::Seconds();
+			double LastUpdateTime = GetStats().LastUpdateStartTime;
+			double ElapsedSinceLastUpdate = GtTime - GetStats().LastUpdateStartTime;
+			// If we're over the latency threshold, wait so the CSV profiler thread can catch up
+			if (ElapsedSinceLastUpdate > MaxUpdateLatency)
+			{
+				bThreadPriorityBoost = true;
+				CSV_EVENT(CsvProfiler, TEXT("GT Stall"));
+				CSV_SCOPED_TIMING_STAT_EXCLUSIVE(CsvProfiler_Stall);
+				CSV_SCOPED_TIMING_STAT(CsvProfiler, Stall);
+				UE_LOG(LogCsvProfiler, Warning, TEXT("Game thread stalled waiting for CSV processing thread to catch up!"));
+				// Stall until the next update
+				do
+				{
+					FPlatformProcess::Sleep(0.f);
+				} while (GetStats().LastUpdateStartTime < GtTime);
+				return true;
+			}
+		}
+
+		// Handle GT stalls triggered from the processing thread
 		if (StallGameThreadCs.TryLock())
 		{
+			bThreadPriorityBoost = false;
 			StallGameThreadCs.Unlock();
-			return;
+			return false;
 		}
 		CSV_SCOPED_TIMING_STAT_EXCLUSIVE(CsvProfiler_Stall);
+		CSV_SCOPED_TIMING_STAT(CsvProfiler, Stall);
 		UE::TScopeLock Lock(StallGameThreadCs);
+		return true;
+	}
+
+	FProcessingThreadStats GetStats()
+	{
+		UE::TScopeLock Lock(StatsCs);
+		return UpdateStats;
 	}
 
 private:
@@ -3414,20 +3421,92 @@ private:
 	}
 	void UpdateThreadPriorityIfNeeded()
 	{
-		EThreadPriority DesiredThreadPriority = GetThreadPriorityFromIndex(GCsvThreadPriorityIndex);
+		int32 ThreadPriIndex = GCsvThreadPriorityIndex;
+		if (bThreadPriorityBoost)
+		{
+			ThreadPriIndex = 5; // TPri_Highest
+		}
+		EThreadPriority DesiredThreadPriority = GetThreadPriorityFromIndex(ThreadPriIndex);
 		if (DesiredThreadPriority != CurrentThreadPriority)
 		{
-			UE_LOG(LogCsvProfiler, Display, TEXT("Setting CSV processing thread priority to %s"), GetThreadPriorityNameFromIndex(GCsvThreadPriorityIndex));
+			UE_LOG(LogCsvProfiler, Display, TEXT("Setting CSV processing thread priority to %s%s"), GetThreadPriorityNameFromIndex(ThreadPriIndex), bThreadPriorityBoost ? TEXT(" - Boost active") : TEXT(""));
 			FPlatformProcess::SetThreadPriority(DesiredThreadPriority);
 			CurrentThreadPriority = DesiredThreadPriority;
 		}
 	}
 
+	void HandleGtStalling(const FCsvProcessThreadDataStats& ProcessingStats, float ElapsedMs)
+	{
+		check(IsInCsvProcessingThread());
+		const float TimeBetweenUpdatesMs = CVarCsvProcessingThreadTimeBetweenUpdates.GetValueOnAnyThread();
+
+		// Detect well-of-despair (stats coming in faster than they can be processed)
+		uint32 NumStatEntriesProcessed = ProcessingStats.CustomStatCount + ProcessingStats.TimestampCount;
+		TotalStatEntriesProcessed += (uint64)NumStatEntriesProcessed;
+		TotalProcessingTime += (double)ElapsedMs / 1000.0;
+
+		// If we've stalled the GT then check if we can resume
+		if (StallCounter > 0)
+		{
+			StallCounter--;
+			if (StallCounter == 0 || ElapsedMs < TimeBetweenUpdatesMs)
+			{
+				if (bHasGtStallLock)
+				{
+					StallGameThreadCs.Unlock();
+					bHasGtStallLock = false;
+				}
+				StallCounter = 0;
+			}
+		}
+		else 
+		{
+			const float StallThresholdMs = CVarCsvProcessingThreadGtStallUpdateTimeThresholdMs.GetValueOnAnyThread();
+			// If this thread is slower than the last and processed more stats, increase SlowUpdateEscalationCount
+			if (StallThresholdMs > 0.0f && PrevElapsedMs > StallThresholdMs && ElapsedMs > PrevElapsedMs && NumStatEntriesProcessed > PrevNumStatEntriesProcessed)
+			{
+				SlowUpdateEscalationCount++;
+				if (SlowUpdateEscalationCount >= CVarCsvProcessingThreadGtStallUpdateEscalationThreshold.GetValueOnAnyThread())
+				{
+					// Stall the GT for up to 2 updates to allow us to catch up
+					StallGameThreadCs.Lock();
+					bHasGtStallLock = true;
+					StallCounter = 2;
+					double StatEntriesProcessedPerSecond = (double)TotalStatEntriesProcessed / TotalProcessingTime;
+					UE_LOG(LogCsvProfiler, Warning, TEXT("Stats coming in faster than we can process them! GT stalled until we can catch up!"));
+					UE_LOG(LogCsvProfiler, Warning, TEXT("Avg processing rate: %.0f stat entries per second (timestamps+custom stats)"), StatEntriesProcessedPerSecond);
+					UE_LOG(LogCsvProfiler, Warning, TEXT("Check CsvProfiler/* stats to see current rates. Run with -csvStatCounts to report per-stat counts"));
+					// Temporarily boost thread priority since the GT is waiting
+					bThreadPriorityBoost = true;
+					SlowUpdateEscalationCount = 0;
+				}
+			}
+			else
+			{
+				SlowUpdateEscalationCount = 0;
+			}
+		}
+
+		PrevNumStatEntriesProcessed = NumStatEntriesProcessed;
+		PrevElapsedMs = ElapsedMs;
+
+		UpdateThreadPriorityIfNeeded();
+	}
 
 
 	FRunnableThread* Thread;
 	FCsvProfiler& CsvProfiler;
+
+	// GT stalling
 	FCriticalSection StallGameThreadCs;
+	uint32 StallCounter = 0;
+	double TotalProcessingTime = 0.0;
+	uint64 TotalStatEntriesProcessed = 0;
+	float PrevElapsedMs = 0.0f;
+	uint32 PrevNumStatEntriesProcessed = 0;
+	int32 SlowUpdateEscalationCount = 0;
+	bool bHasGtStallLock = false;
+	volatile bool bThreadPriorityBoost = false;
 };
 
 void FCsvProfilerThreadDataProcessor::Process(FCsvProcessThreadDataStats& OutStats, int32& OutMinFrameNumberProcessed, TSharedPtr<FCsvThreadGroupStatProcessor> TaskWorkerThreadGroupStatProcessor)
@@ -3984,6 +4063,7 @@ void FCsvProfiler::BeginCaptureInternal(const FCsvCaptureCommand& CurrentCommand
 	SetMetadataInternal(TEXT("TargetFramerate"), *FString::FromInt(TargetFPS));
 	SetMetadataInternal(TEXT("StartTimestamp"), *FString::Printf(TEXT("%lld"), FDateTime::UtcNow().ToUnixTimestamp()));
 	SetMetadataInternal(TEXT("NamedEvents"), (GCycleStatsShouldEmitNamedEvents > 0) ? TEXT("1") : TEXT("0"));
+	SetMetadataInternal(TEXT("HasGtStalls"), TEXT("0"));
 
 	static IConsoleVariable* VSyncCVar = IConsoleManager::Get().FindConsoleVariable(TEXT("r.VSync"));
 	if (VSyncCVar)
@@ -4024,10 +4104,20 @@ void FCsvProfiler::EndFrame()
 	check(IsInGameThread());
 	if (GCsvProfilerIsCapturing)
 	{
-		// If the GT is stalled, wait
 		if (ProcessingThread)
 		{
-			ProcessingThread->StallGameThreadIfNeeded();
+			FProcessingThreadStats ProcessingThreadStats = ProcessingThread->GetStats();
+
+			double ElapsedSinceLastUpdate = FPlatformTime::Seconds()-ProcessingThreadStats.LastUpdateStartTime;
+
+			CSV_CUSTOM_STAT(CsvProfiler, LastUpdateMemUsedMB, (float)ProcessingThreadStats.ProcessingStats.GetEstimatedMemoryUse() / 1024.0f / 1024.0f, ECsvCustomStatOp::Set);
+			CSV_CUSTOM_STAT(CsvProfiler, TimeSinceLastUpdateMs, (float)ElapsedSinceLastUpdate*1000.0f, ECsvCustomStatOp::Set);
+
+			// Stall the GT if needed
+			if ( ProcessingThread->StallGameThreadIfNeeded() )
+			{
+				CSV_METADATA(TEXT("HasGtStalls"), TEXT("1"));
+			}
 		}		
 
 		OnCSVProfileEndFrameDelegate.Broadcast();
@@ -5242,13 +5332,13 @@ bool FCsvProfiler::IsCapturing_Renderthread()
 	return GCsvProfilerIsCapturingRT;
 }
 
-float FCsvProfiler::ProcessStatData(uint32* OutNumStatEntriesProcessed)
+float FCsvProfiler::ProcessStatData(FCsvProcessThreadDataStats* OutStats)
 {
 	check(IsInCsvProcessingThread());
 
 	QUICK_SCOPE_CYCLE_COUNTER(STAT_FCsvProfiler_ProcessStatData);
 
-	float ElapsedMS = 0.0f;
+	float ElapsedMs = 0.0f;
 	if (!IsShuttingDown.GetValue())
 	{
 		double StartTime = FPlatformTime::Seconds();
@@ -5258,18 +5348,18 @@ float FCsvProfiler::ProcessStatData(uint32* OutNumStatEntriesProcessed)
 		{
 			CsvWriter->Process(Stats);
 		}
-		ElapsedMS = float(FPlatformTime::Seconds() - StartTime) * 1000.0f;
+		ElapsedMs = float(FPlatformTime::Seconds() - StartTime) * 1000.0f;
 		CSV_CUSTOM_STAT(CsvProfiler, NumTimestampsProcessed, (int32)Stats.TimestampCount, ECsvCustomStatOp::Accumulate);
 		CSV_CUSTOM_STAT(CsvProfiler, NumCustomStatsProcessed, (int32)Stats.CustomStatCount, ECsvCustomStatOp::Accumulate);
 		CSV_CUSTOM_STAT(CsvProfiler, NumEventsProcessed, (int32)Stats.EventCount, ECsvCustomStatOp::Accumulate);
-		CSV_CUSTOM_STAT(CsvProfiler, ProcessCSVStats, ElapsedMS, ECsvCustomStatOp::Accumulate);
+		CSV_CUSTOM_STAT(CsvProfiler, ProcessCSVStats, ElapsedMs, ECsvCustomStatOp::Accumulate);
 
-		if (OutNumStatEntriesProcessed)
+		if (OutStats)
 		{
-			*OutNumStatEntriesProcessed = Stats.CustomStatCount + Stats.TimestampCount;
+			*OutStats = Stats;
 		}		
 	}
-	return ElapsedMS;
+	return ElapsedMs;
 }
 
 TCsvPersistentCustomStat<int32>* FCsvProfiler::GetOrCreatePersistentCustomStatInt(FName Name, int32 CategoryIndex, bool bResetEachFrame)
diff --git a/Engine/Source/Runtime/Core/Public/ProfilingDebugging/CsvProfiler.h b/Engine/Source/Runtime/Core/Public/ProfilingDebugging/CsvProfiler.h
--- a/Engine/Source/Runtime/Core/Public/ProfilingDebugging/CsvProfiler.h
+++ b/Engine/Source/Runtime/Core/Public/ProfilingDebugging/CsvProfiler.h
@@ -397,6 +397,17 @@ protected:
 	std::atomic<int> Value;
 };
 
+/**
+* Processing stats
+*/
+struct FCsvProcessThreadDataStats
+{
+	uint32 TimestampCount = 0;
+	uint32 CustomStatCount = 0;
+	uint32 EventCount = 0;
+
+	uint64 GetEstimatedMemoryUse() const;
+};
 
 /**
 * FCsvProfiler class. This manages recording and reporting all for CSV stats
@@ -595,7 +606,7 @@ private:
 
 	void FinalizeCsvFile();
 
-	float ProcessStatData(uint32* OutNumStatEntriesProcessed = nullptr);
+	float ProcessStatData(FCsvProcessThreadDataStats* OutNumStatEntriesProcessed = nullptr);
 
 	void BeginCaptureInternal(const FCsvCaptureCommand& CurrentCommand);
 
