Logging: Removed deprecated log template API

#rnx
#Synced-CL 51340251

diff --git a/Engine/Source/Runtime/Core/Public/Logging/StructuredLogFormat.h b/Engine/Source/Runtime/Core/Public/Logging/StructuredLogFormat.h
--- a/Engine/Source/Runtime/Core/Public/Logging/StructuredLogFormat.h
+++ b/Engine/Source/Runtime/Core/Public/Logging/StructuredLogFormat.h
@@ -5,7 +5,6 @@
 #include "Containers/Array.h"
 #include "Containers/StringFwd.h"
 #include "Internationalization/Text.h"
-#include "Misc/ScopeExit.h"
 #include "Templates/FunctionWithContext.h"
 #include "Templates/UniquePtr.h"
 
@@ -64,49 +63,6 @@ private:
 	TUniquePtr<uint8[]> Data;
 };
 
-/** Store a log template in a detachable allocation from FMemory. */
-class FMemoryLogTemplateStorage
-{
-public:
-	FMemoryLogTemplateStorage() = default;
-	FMemoryLogTemplateStorage(const FMemoryLogTemplateStorage&) = delete;
-	FMemoryLogTemplateStorage& operator=(const FMemoryLogTemplateStorage&) = delete;
-
-	~FMemoryLogTemplateStorage()
-	{
-		Free(Data);
-	}
-
-	inline void* Allocate(int32 Size)
-	{
-		Free(Data);
-		Data = FMemory::Malloc(Size);
-		return Data;
-	}
-
-	inline const void* Get() const
-	{
-		return Data;
-	}
-
-	inline void* Detach()
-	{
-		ON_SCOPE_EXIT { Data = nullptr; };
-		return Data;
-	}
-
-	inline static void Free(void* D)
-	{
-		if (D)
-		{
-			FMemory::Free(D);
-		}
-	}
-
-private:
-	void* Data = nullptr;
-};
-
 UE_API void CreateLogTemplate(const TCHAR* Format, const FLogTemplateOptions& Options, const FLogField* Fields, int32 FieldCount, TFunctionWithContext<void* (int32)> Allocate);
 UE_API void CreateLogTemplate(const UTF8CHAR* Format, const FLogTemplateOptions& Options, const FLogField* Fields, int32 FieldCount, TFunctionWithContext<void* (int32)> Allocate);
 UE_API void CreateLocalizedLogTemplate(const FText& Format, const FLogTemplateOptions& Options, const FLogField* Fields, int32 FieldCount, TFunctionWithContext<void* (int32)> Allocate);
@@ -204,31 +160,6 @@ using FInlineLogTemplate = TLogTemplate<Logging::Private::FInlineLogTemplateStor
 /** A log template that stores its data on the heap. A reasonable default when for non-temporary templates. */
 using FUniqueLogTemplate = TLogTemplate<Logging::Private::FUniqueLogTemplateStorage>;
 
-UE_DEPRECATED(5.6, "Use FInlineLogTemplate or FUniqueLogTemplate.")
-inline FLogTemplate* CreateLogTemplate(const TCHAR* Format, const FLogTemplateOptions& Options = {})
-{
-	return TLogTemplate<Logging::Private::FMemoryLogTemplateStorage>(Format, Options).Detach();
-}
-
-UE_DEPRECATED(5.6, "Use FInlineLogTemplate or FUniqueLogTemplate.")
-inline FLogTemplate* CreateLogTemplate(const FText& Format, const FLogTemplateOptions& Options = {})
-{
-	return TLogTemplate<Logging::Private::FMemoryLogTemplateStorage>(Format, Options).Detach();
-}
-
-UE_DEPRECATED(5.6, "Use FInlineLogTemplate or FUniqueLogTemplate.")
-inline FLogTemplate* CreateLogTemplate(const TCHAR* TextNamespace, const TCHAR* TextKey, const TCHAR* Format, const FLogTemplateOptions& Options = {})
-{
-	return TLogTemplate<Logging::Private::FMemoryLogTemplateStorage>(TextNamespace, TextKey, Format, Options).Detach();
-}
-
-UE_DEPRECATED(5.6, "Use FInlineLogTemplate or FUniqueLogTemplate.")
-inline void DestroyLogTemplate(FLogTemplate* Template)
-{
-	Logging::Private::DestroyLogTemplate(Template);
-	Logging::Private::FMemoryLogTemplateStorage::Free(Template);
-}
-
 /**
  * Serializes a localized log format template to compact binary. Call from SerializeForLog.
  *
