Csv profiler: Disable memory tracking by default because it takes 0.5ms on some platforms. It can be enabled via the cvar csv.MemoryTracking
Fix undercounting with memory tracking. Count the writer series data etc instead of just the unprocessed stat data
Rename stats LastUpdateMemUsedMB->MemUsedLastUpdateMB and TimeSinceLastUpdateMs->UpdateLatencyMs

#Synced-CL 51245126
#Synced-CL 51245126

diff --git a/Engine/Source/Runtime/Core/Private/ProfilingDebugging/CsvProfiler.cpp b/Engine/Source/Runtime/Core/Private/ProfilingDebugging/CsvProfiler.cpp
--- a/Engine/Source/Runtime/Core/Private/ProfilingDebugging/CsvProfiler.cpp
+++ b/Engine/Source/Runtime/Core/Private/ProfilingDebugging/CsvProfiler.cpp
@@ -298,6 +298,13 @@ FAutoConsoleVariableRef CVarCsvThreadPriority(
 	TEXT("This is switchable at runtime"),
 	ECVF_Default);
 
+TAutoConsoleVariable<int32> CVarCsvTrackMemoryUse(
+	TEXT("csv.TrackMemoryUse"),
+	0,
+	TEXT("if 1, tracks memory use via the CsvProfiler/MemUsedLastUpdateMB stat. This can take 0.5ms on the processing thread so is disabled by default"),
+	ECVF_Default
+);
+
 
 static bool GCsvUseProcessingThread = true;
 static int32 GCsvRepeatCount = 0;
@@ -2345,13 +2352,6 @@ struct FCsvAggregateStatSeries : public FCsvStatSeries
 	TMap<int64, FCsvStatSeriesValue> RowValues;
 };
 
-
-uint64 FCsvProcessThreadDataStats::GetEstimatedMemoryUse() const
-{
-	// Note: this doesn't include memory for event strings
-	return sizeof(FCsvTimingMarker)*TimestampCount + sizeof(FCsvCustomStat)*CustomStatCount + sizeof(FCsvEvent)*EventCount;
-}
-
 class FCsvThreadGroupStatProcessor;
 
 class FCsvStreamWriter
@@ -3179,6 +3179,13 @@ void FCsvStreamWriter::Process(FCsvProcessThreadDataStats& OutStats)
 			}
 		}
 	}
+	// Read the peak mem use, if configured
+	OutStats.MemoryUseBytes = 0;
+	if (CVarCsvTrackMemoryUse.GetValueOnAnyThread(true))
+	{
+		CSV_SCOPED_TIMING_STAT(CsvProfiler, GetMemoryUse);
+		OutStats.MemoryUseBytes = GetAllocatedSize();
+	}	
 
 	int32 MinFrameNumberProcessed = MAX_int32;
 	{
@@ -4107,11 +4114,13 @@ void FCsvProfiler::EndFrame()
 		if (ProcessingThread)
 		{
 			FProcessingThreadStats ProcessingThreadStats = ProcessingThread->GetStats();
-
 			double ElapsedSinceLastUpdate = FPlatformTime::Seconds()-ProcessingThreadStats.LastUpdateStartTime;
 
-			CSV_CUSTOM_STAT(CsvProfiler, LastUpdateMemUsedMB, (float)ProcessingThreadStats.ProcessingStats.GetEstimatedMemoryUse() / 1024.0f / 1024.0f, ECsvCustomStatOp::Set);
-			CSV_CUSTOM_STAT(CsvProfiler, TimeSinceLastUpdateMs, (float)ElapsedSinceLastUpdate*1000.0f, ECsvCustomStatOp::Set);
+			if (ProcessingThreadStats.ProcessingStats.MemoryUseBytes>0)
+			{
+				CSV_CUSTOM_STAT(CsvProfiler, MemUsedLastUpdateMB, (float)ProcessingThreadStats.ProcessingStats.MemoryUseBytes / 1024.0f / 1024.0f, ECsvCustomStatOp::Set);
+			}
+			CSV_CUSTOM_STAT(CsvProfiler, UpdateLatencyMs, (float)ElapsedSinceLastUpdate*1000.0f, ECsvCustomStatOp::Set);
 
 			// Stall the GT if needed
 			if ( ProcessingThread->StallGameThreadIfNeeded() )
diff --git a/Engine/Source/Runtime/Core/Public/ProfilingDebugging/CsvProfiler.h b/Engine/Source/Runtime/Core/Public/ProfilingDebugging/CsvProfiler.h
--- a/Engine/Source/Runtime/Core/Public/ProfilingDebugging/CsvProfiler.h
+++ b/Engine/Source/Runtime/Core/Public/ProfilingDebugging/CsvProfiler.h
@@ -405,8 +405,7 @@ struct FCsvProcessThreadDataStats
 	uint32 TimestampCount = 0;
 	uint32 CustomStatCount = 0;
 	uint32 EventCount = 0;
-
-	uint64 GetEstimatedMemoryUse() const;
+	uint64 MemoryUseBytes = 0;
 };
 
 /**
