[Core] Fixup narrowing conversion warnings in compact set

#tests compiled foundation tests
#Synced-CL 51344084

diff --git a/Engine/Source/Runtime/Core/Public/Containers/CompactHashTable.h b/Engine/Source/Runtime/Core/Public/Containers/CompactHashTable.h
--- a/Engine/Source/Runtime/Core/Public/Containers/CompactHashTable.h
+++ b/Engine/Source/Runtime/Core/Public/Containers/CompactHashTable.h
@@ -88,7 +88,7 @@ namespace UE::Core::CompactHashTable
 	}
 
 	/* Calculate the size of the hash table from the number of elements in the set */
-	[[nodiscard]] inline constexpr size_t GetHashCount(uint32 NumElements)
+	[[nodiscard]] inline constexpr uint32 GetHashCount(uint32 NumElements)
 	{
 		if (!NumElements)
 		{
diff --git a/Engine/Source/Runtime/Core/Public/Containers/CompactSetBase.h b/Engine/Source/Runtime/Core/Public/Containers/CompactSetBase.h
--- a/Engine/Source/Runtime/Core/Public/Containers/CompactSetBase.h
+++ b/Engine/Source/Runtime/Core/Public/Containers/CompactSetBase.h
@@ -114,7 +114,7 @@ protected:
 	}
 
 	/* Calculate the size of the hash table from the number of elements in the set */
-	[[nodiscard]] UE_FORCEINLINE_HINT static constexpr SizeType GetHashCount(uint32 NumElements)
+	[[nodiscard]] UE_FORCEINLINE_HINT static constexpr HashCountType GetHashCount(uint32 NumElements)
 	{
 		return UE::Core::CompactHashTable::GetHashCount(NumElements);
 	}
@@ -126,7 +126,7 @@ protected:
 	}
 
 	/* Get the total memory required for the compact set for the given number of elements and size of the number of hash elements */
-	[[nodiscard]] UE_FORCEINLINE_HINT static constexpr size_t GetTotalMemoryRequiredInBytes(uint32 NumElements, uint32 HashCount, const FCompactSetLayout Layout)
+	[[nodiscard]] UE_FORCEINLINE_HINT static constexpr size_t GetTotalMemoryRequiredInBytes(uint32 NumElements, HashCountType HashCount, const FCompactSetLayout Layout)
 	{
 		return NumElements ? GetElementsSizeInBytes(NumElements, Layout) + UE::Core::CompactHashTable::GetMemoryRequiredInBytes(NumElements, HashCount) + HashCountSize : 0;
 	}
@@ -138,7 +138,7 @@ protected:
 	}
 
 	/* The total amount of memory required by the has set to store N elements with a hash table size of Y */
-	[[nodiscard]] static constexpr SizeType GetMaxElementsForAvailableSpace(size_t TotalBytes, uint32 HashCount, uint32 MinElementCount, const FCompactSetLayout Layout)
+	[[nodiscard]] static constexpr SizeType GetMaxElementsForAvailableSpace(size_t TotalBytes, HashCountType HashCount, uint32 MinElementCount, const FCompactSetLayout Layout)
 	{
 		// Given some space in memory and a requested size for hash table, figure out how many elements we can fit in the remaining space
 		const uint32 TypeSize = UE::Core::CompactHashTable::GetTypeSize(MinElementCount);
@@ -146,25 +146,27 @@ protected:
 		const size_t AvailableBytes = TotalBytes - sizeof(HashCountType) - ((size_t)HashCount << TypeShift); // Remove hashtable and HashCount data
 		const size_t MaxElements = AvailableBytes / (Layout.Size + TypeSize); // Calculate the max available ignoring the fact that the hash data has to be aligned
 		const size_t RealAvailableBytes = AlignDown(AvailableBytes - (MaxElements << TypeShift), UE::Core::CompactHashTable::GetMemoryAlignment()); // Remove the max required indexes and align down
-		return FMath::Min<SizeType>(MaxElements, RealAvailableBytes / Layout.Size); // Now we can get the true number of elements we could potentially use within the aligned space
+		return (SizeType)FMath::Min<size_t>(MaxElements, RealAvailableBytes / Layout.Size); // Now we can get the true number of elements we could potentially use within the aligned space
 	}
 
 	[[nodiscard]] SizeType AllocatorCalculateSlackGrow(SizeType NewMaxElements, const FCompactSetLayout& Layout) const
 	{
-		const size_t OldHashCount = MaxElements > 0 ? *(HashCountType*)(Elements.GetAllocation() + GetElementsSizeInBytes(MaxElements, Layout)) : 0;
-		const size_t OldSize = MaxElements > 0 ? GetTotalMemoryRequiredInBytes(MaxElements, OldHashCount, Layout) : 0;
+		const HashCountType OldHashCount = MaxElements > 0 ? *(HashCountType*)(Elements.GetAllocation() + GetElementsSizeInBytes(MaxElements, Layout)) : 0;
+		const size_t OldSize = GetTotalMemoryRequiredInBytes(MaxElements, OldHashCount, Layout);
 
-		const size_t NewHashCount = NewMaxElements > 0 ? GetHashCount(NewMaxElements) : 0;
-		const size_t NewSize = NewMaxElements > 0 ? GetTotalMemoryRequiredInBytes(NewMaxElements, NewHashCount, Layout) : 0;
+		const HashCountType NewHashCount = GetHashCount(NewMaxElements);
+		const size_t NewSize = GetTotalMemoryRequiredInBytes(NewMaxElements, NewHashCount, Layout);
 		size_t NewSlackSize = 0;
 
+		checkf(NewSize > 0 && NewSize <= MAX_int32, TEXT("Invalid size for TSet[%zd]: NewMaxElements[%d] ElementSize[%d] HashCount[%d]"), NewSize, NewMaxElements, Layout.Size, NewHashCount);
+
 		if constexpr (TAllocatorTraits<AllocatorType>::SupportsElementAlignment)
 		{
-			NewSlackSize = Elements.CalculateSlackGrow(NewSize, OldSize, 1, Layout.Alignment);
+			NewSlackSize = Elements.CalculateSlackGrow((SizeType)NewSize, (SizeType)OldSize, 1, Layout.Alignment);
 		}
 		else
 		{
-			NewSlackSize = Elements.CalculateSlackGrow(NewSize, OldSize, 1);
+			NewSlackSize = Elements.CalculateSlackGrow((SizeType)NewSize, (SizeType)OldSize, 1);
 		}
 
 		if (NewSlackSize == NewSize)
@@ -182,7 +184,7 @@ protected:
 			return NewMaxElements;
 		}
 
-		size_t SlackHashCount = GetHashCount(SlackNumElements);
+		HashCountType SlackHashCount = GetHashCount(SlackNumElements);
 		if (SlackHashCount > NewHashCount)
 		{
 			// New Size required to much space. Cut the new hash count in half and restrict the elements to the max for that hash count
@@ -208,24 +210,24 @@ protected:
 
 		if (NewMaxElements != MaxElements)
 		{
-			const size_t OldHashCount = MaxElements > 0 ? *(HashCountType*)(Elements.GetAllocation() + GetElementsSizeInBytes(MaxElements, Layout)) : 0;
+			const HashCountType OldHashCount = MaxElements > 0 ? *(HashCountType*)(Elements.GetAllocation() + GetElementsSizeInBytes(MaxElements, Layout)) : 0;
 			const size_t OldSize = MaxElements > 0 ? GetTotalMemoryRequiredInBytes(MaxElements, OldHashCount, Layout) : 0;
 
-			const size_t NewHashCount = NewMaxElements > 0 ? UE::Core::CompactHashTable::GetHashCount(NewMaxElements) : 0;
+			const HashCountType NewHashCount = NewMaxElements > 0 ? UE::Core::CompactHashTable::GetHashCount(NewMaxElements) : 0;
 			const size_t NewSize = NewMaxElements > 0 ? GetTotalMemoryRequiredInBytes(NewMaxElements, NewHashCount, Layout) : 0;
 
-			checkf(NewSize >= 0 && NewSize <= MAX_int32, TEXT("Invalid size for TSet[%lld]: NewMaxElements[%d] ElementSize[%d] HashCount[%lld]"), NewSize, NewMaxElements, Layout.Size, NewHashCount);
+			checkf(NewSize >= 0 && NewSize <= MAX_int32, TEXT("Invalid size for TSet[%lld]: NewMaxElements[%d] ElementSize[%d] HashCount[%d]"), NewSize, NewMaxElements, Layout.Size, NewHashCount);
 
 			if (bPreserve && NewMaxElements > MaxElements && OldHashCount == NewHashCount)
 			{
 				// If preserving then we copy over all the data
 				if constexpr (TAllocatorTraits<AllocatorType>::SupportsElementAlignment)
 				{
-					Elements.ResizeAllocation(OldSize, NewSize, 1, Layout.Alignment);
+					Elements.ResizeAllocation((SizeType)OldSize, (SizeType)NewSize, 1, Layout.Alignment);
 				}
 				else
 				{
-					Elements.ResizeAllocation(OldSize, NewSize, 1);
+					Elements.ResizeAllocation((SizeType)OldSize, (SizeType)NewSize, 1);
 				}
 
 				const uint32 NewTypeShift = UE::Core::CompactHashTable::GetTypeShift(NewMaxElements);
@@ -250,11 +252,11 @@ protected:
 				// If not preserving (or shrinking) then we only need to copy the array data
 				if constexpr (TAllocatorTraits<AllocatorType>::SupportsElementAlignment)
 				{
-					Elements.ResizeAllocation(NumElements * Layout.Size, NewSize, 1, Layout.Alignment);
+					Elements.ResizeAllocation(NumElements * Layout.Size, (SizeType)NewSize, 1, Layout.Alignment);
 				}
 				else
 				{
-					Elements.ResizeAllocation(NumElements * Layout.Size, NewSize, 1);
+					Elements.ResizeAllocation(NumElements * Layout.Size, (SizeType)NewSize, 1);
 				}
 
 				// Update hash size
