PR #14237: [Containers][Set] Symmetric difference helper for compact and sparse sets

#tests Compiled

#Synced-CL 51398680

diff --git a/Engine/Source/Runtime/Core/Public/Containers/CompactSet.h.inl b/Engine/Source/Runtime/Core/Public/Containers/CompactSet.h.inl
--- a/Engine/Source/Runtime/Core/Public/Containers/CompactSet.h.inl
+++ b/Engine/Source/Runtime/Core/Public/Containers/CompactSet.h.inl
@@ -1246,6 +1246,33 @@ public:
 		}
 		return Result;
 	}
+	
+	/** @return the symmetric difference of two sets. (A not in B and B not in A where A is this and B is Other) */
+	[[nodiscard]] UE_TCOMPACT_SET SymmetricDifference(const UE_TCOMPACT_SET& OtherSet) const
+	{
+		UE_TCOMPACT_SET Result;
+		Result.Reserve(this->NumElements + OtherSet.NumElements); // Worst case: no overlap
+
+		// Add all elements from A that are not in B
+		for (const ElementType& Element : *this)
+		{
+			if (!OtherSet.Contains(KeyFuncs::GetSetKey(Element)))
+			{
+				Result.Add(Element);
+			}
+		}
+
+		// Add all elements from B that are not in A
+		for (const ElementType& Element : OtherSet)
+		{
+			if (!this->Contains(KeyFuncs::GetSetKey(Element)))
+			{
+				Result.Add(Element);
+			}
+		}
+
+		return Result;
+	}
 
 	/**
 	 * Determine whether the specified set is entirely included within this set
diff --git a/Engine/Source/Runtime/Core/Public/Containers/SparseSet.h.inl b/Engine/Source/Runtime/Core/Public/Containers/SparseSet.h.inl
--- a/Engine/Source/Runtime/Core/Public/Containers/SparseSet.h.inl
+++ b/Engine/Source/Runtime/Core/Public/Containers/SparseSet.h.inl
@@ -1133,6 +1133,33 @@ public:
 		}
 		return Result;
 	}
+	
+	/** @return the symmetric difference of two sets. (A not in B and B not in A where A is this and B is Other) */
+	[[nodiscard]] UE_TSPARSE_SET SymmetricDifference(const UE_TSPARSE_SET& OtherSet) const
+	{
+		UE_TSPARSE_SET Result;
+		Result.Reserve(Num() + OtherSet.Num()); // Worst case is no overlap
+
+		// Add elements from A that are not in B
+		for(TConstIterator SetIt(*this);SetIt;++SetIt)
+		{
+			if(!OtherSet.Contains(KeyFuncs::GetSetKey(*SetIt)))
+			{
+				Result.Add(*SetIt);
+			}
+		}
+
+		// Add elements from B that are not in A
+		for(TConstIterator SetIt(OtherSet);SetIt;++SetIt)
+		{
+			if(!this->Contains(KeyFuncs::GetSetKey(*SetIt)))
+			{
+				Result.Add(*SetIt);
+			}
+		}
+
+		return Result;
+	}
 
 	/**
 	 * Determine whether the specified set is entirely included within this set
diff --git a/Engine/Source/Runtime/Core/Tests/Containers/SetTests.cpp b/Engine/Source/Runtime/Core/Tests/Containers/SetTests.cpp
--- a/Engine/Source/Runtime/Core/Tests/Containers/SetTests.cpp
+++ b/Engine/Source/Runtime/Core/Tests/Containers/SetTests.cpp
@@ -199,6 +199,7 @@ struct TSetTestHelper
 		SECTION("Advanced Tests")
 		{
 			int32 IntTestItems[256];
+			int32 IntExtraTestItem[16];
 
 			{
 				int32 IntInit = 1;
@@ -206,51 +207,112 @@ struct TSetTestHelper
 				{
 					Int = HashCombine(IntInit++, 0);
 				}
+				for (int32& Int : IntExtraTestItem)
+				{
+					Int = HashCombine(IntInit++, 0);
+				}
 			}
 
-			SetType<int32> Set1(IntTestItems);
-			SetType<int32> Set2;
+			// SetAllTestItems and SetAllTestItemsReverse are the same, but different order
+			SetType<int32> SetAllTestItems(IntTestItems);
+			CHECK(SetAllTestItems.Num() == UE_ARRAY_COUNT(IntTestItems)); // Ensure we didn't have any duplicates
 
-			for (int32 Item : ReverseIterate(IntTestItems))
+			SetType<int32> SetAllTestItemsReverse;
+
+			// Test behavior of same set of data but in a different order
 			{
-				Set2.Add(Item);
+				for (int32 Item : ReverseIterate(IntTestItems))
+				{
+					SetAllTestItemsReverse.Add(Item);
+				}
+
+				for (int32 TestItem : IntTestItems)
+				{
+					SetTestExists(TestItem, SetAllTestItems);
+					SetTestExists(TestItem, SetAllTestItemsReverse);
+				}
+
+				CHECK(SetAllTestItems.Includes(SetAllTestItemsReverse));
+				CHECK(SetAllTestItemsReverse.Includes(SetAllTestItems));
+				CHECK(SetAllTestItems.Difference(SetAllTestItemsReverse).IsEmpty());
+
+				CHECK(SetAllTestItems.Union(SetAllTestItemsReverse).Num() == SetAllTestItems.Num());
+
+				CHECK(SetAllTestItems.Union(SetAllTestItemsReverse).Difference(SetAllTestItemsReverse).IsEmpty());
+				CHECK(SetAllTestItems.Intersect(SetAllTestItemsReverse).Difference(SetAllTestItemsReverse).IsEmpty());
 			}
 
-			for (int32 TestItem : IntTestItems)
+			// Contains the first half of whole set of items
+			SetType<int32> SetFirstHalfTestItems;
+			// Contains the second half of whole set of items
+			SetType<int32> SetSecondHalfTestItems;
+
+			// Test behavior of subsets of the same data
 			{
-				SetTestExists(TestItem, Set1);
-				SetTestExists(TestItem, Set2);
+				for (int32 Index = 0; Index < UE_ARRAY_COUNT(IntTestItems) / 2; ++Index)
+				{
+					SetFirstHalfTestItems.Add(IntTestItems[Index]);
+				}
+				for (int32 Index = UE_ARRAY_COUNT(IntTestItems) / 2; Index < UE_ARRAY_COUNT(IntTestItems); ++Index)
+				{
+					SetSecondHalfTestItems.Add(IntTestItems[Index]);
+				}
+
+				CHECK(SetAllTestItems.Includes(SetFirstHalfTestItems));
+				CHECK(SetAllTestItems.Includes(SetSecondHalfTestItems));
+				CHECK(!SetFirstHalfTestItems.Includes(SetAllTestItems));
+				CHECK(!SetSecondHalfTestItems.Includes(SetAllTestItems));
+
+				CHECK(SetAllTestItems.Difference(SetFirstHalfTestItems).Num() == SetAllTestItems.Num() - SetFirstHalfTestItems.Num());
+				CHECK(SetAllTestItems.Difference(SetSecondHalfTestItems).Num() == SetAllTestItems.Num() - SetSecondHalfTestItems.Num());
+				CHECK(SetAllTestItems.Intersect(SetFirstHalfTestItems).Difference(SetFirstHalfTestItems).IsEmpty());
+				CHECK(SetAllTestItems.Difference(SetFirstHalfTestItems).Difference(SetAllTestItemsReverse.Difference(SetFirstHalfTestItems)).IsEmpty());
+
+				CHECK(SetAllTestItems.Includes(SetFirstHalfTestItems));
+
+				CHECK(SetFirstHalfTestItems.Difference(SetSecondHalfTestItems).Num() == SetFirstHalfTestItems.Num());
+				CHECK(SetSecondHalfTestItems.Difference(SetFirstHalfTestItems).Num() == SetSecondHalfTestItems.Num());
+				CHECK(SetFirstHalfTestItems.SymmetricDifference(SetSecondHalfTestItems).Num() == UE_ARRAY_COUNT(IntTestItems));
+				CHECK(SetSecondHalfTestItems.SymmetricDifference(SetFirstHalfTestItems).Num() == UE_ARRAY_COUNT(IntTestItems));
+				CHECK(SetFirstHalfTestItems.SymmetricDifference(SetAllTestItems).Num() == SetFirstHalfTestItems.Num());
+				CHECK(SetSecondHalfTestItems.SymmetricDifference(SetAllTestItemsReverse).Num() == SetSecondHalfTestItems.Num());
+				CHECK(SetFirstHalfTestItems.Intersect(SetSecondHalfTestItems).IsEmpty());
 			}
 
-			CHECK(Set1.Includes(Set2));
-			CHECK(Set2.Includes(Set1));
-			CHECK(Set1.Difference(Set2).IsEmpty());
+			// Test behavior of disinct sets of data
+			{
+				// More items that are not contained in any other set
+				SetType<int32> SetExtraItems(IntExtraTestItem);
+				CHECK(SetExtraItems.Num() == UE_ARRAY_COUNT(IntExtraTestItem)); // Ensure we didn't have any duplicates
+				CHECK(SetExtraItems.Union(SetAllTestItems).Num() == SetExtraItems.Num() + SetAllTestItems.Num()); // Ensure we didn't have any duplicates between the two sets
+				CHECK(SetExtraItems.SymmetricDifference(SetAllTestItems).Num() == SetExtraItems.Num() + SetAllTestItems.Num());
 
-			CHECK(Set1.Union(Set2).Difference(Set2).IsEmpty());
-			CHECK(Set1.Intersect(Set2).Difference(Set2).IsEmpty());
+				SetType<int32> SetExtraItemsWithFirstHalf = SetExtraItems.Union(SetFirstHalfTestItems);
+				SetType<int32> SetExtraItemsWithSecondHalf = SetExtraItems.Union(SetSecondHalfTestItems);
 
-			SetType<int32> Set3;
-			for (int32 Index = 0; Index < UE_ARRAY_COUNT(IntTestItems) / 2; ++Index)
-			{
-				Set3.Add(IntTestItems[Index]);
-			}
+				CHECK(SetExtraItemsWithFirstHalf.Includes(SetExtraItems));
+				CHECK(SetExtraItemsWithFirstHalf.Includes(SetFirstHalfTestItems));
+				CHECK(SetExtraItemsWithSecondHalf.Includes(SetExtraItems));
+				CHECK(SetExtraItemsWithSecondHalf.Includes(SetSecondHalfTestItems));
 
-			CHECK(Set1.Difference(Set3).Num() == UE_ARRAY_COUNT(IntTestItems) / 2);
-			CHECK(Set1.Intersect(Set3).Difference(Set3).IsEmpty());
-			CHECK(Set1.Difference(Set3).Difference(Set2.Difference(Set3)).IsEmpty());
+				CHECK(SetExtraItemsWithFirstHalf.Intersect(SetExtraItemsWithSecondHalf).Num() == SetExtraItems.Num());
 
-			CHECK(Set1.Includes(Set3));
-			CHECK(!Set3.Includes(Set1));
+				CHECK(SetExtraItemsWithFirstHalf.SymmetricDifference(SetAllTestItems).Num() == SetExtraItems.Num() + SetAllTestItems.Num() - SetFirstHalfTestItems.Num());
+				CHECK(SetExtraItemsWithFirstHalf.SymmetricDifference(SetExtraItemsWithSecondHalf).Num() == SetAllTestItems.Num());
+
+				CHECK(SetExtraItems.Union(SetFirstHalfTestItems).SymmetricDifference(SetSecondHalfTestItems).Num() == SetExtraItems.Num() + SetAllTestItems.Num());
+				CHECK(SetExtraItemsWithFirstHalf.SymmetricDifference(SetExtraItemsWithSecondHalf).Intersect(SetExtraItems).IsEmpty());
+			}
 		}
 
 		SECTION("Copy/Move")
 		{
 			FString StringTestItems[] = { TEXT("Zero"), TEXT("One"), TEXT("Two"), TEXT("Three") };
-			SetType<FString> Set1( { StringTestItems[0], StringTestItems[1], StringTestItems[2], StringTestItems[3] });
-			SetType<FString> Set2(MoveTemp(Set1));
+			SetType<FString> SetAllTestItems( { StringTestItems[0], StringTestItems[1], StringTestItems[2], StringTestItems[3] });
+			SetType<FString> SetAllTestItemsReverse(MoveTemp(SetAllTestItems));
 
-			CHECK(Set1.IsEmpty());
-			CHECK(Set1.Difference(Set2).IsEmpty());
+			CHECK(SetAllTestItems.IsEmpty());
+			CHECK(SetAllTestItems.Difference(SetAllTestItemsReverse).IsEmpty());
 		}
 
 		SECTION("Growth Tests")
@@ -287,7 +349,8 @@ struct TSetTestHelper
 			// + 1 so it's immune to < vs <=, we want to make sure to test both unhash and memset scenarios for both set types
 			const uint32 MaxElements = (PartialElements + 1) * FMath::Max(UE_SPARSESET_UNHASH_THRESHOLD, UE_COMPACTSET_UNHASH_THRESHOLD);
 			Set.Empty(MaxElements);
-			CHECK(Set.Max() == MaxElements);
+			CHECK(Set.Max() >= MaxElements);
+			const uint32 CalculatedMaxElements = Set.Max();
 
 			// There's not really much to test here, just run multiple passes of reusing the set to make sure things don't explode
 			for (uint32 LoopIndex = 0; LoopIndex < 4; ++LoopIndex)
@@ -303,7 +366,7 @@ struct TSetTestHelper
 				CHECK(Set.Num() == PartialElements);
 				Set.Reset();
 				CHECK(Set.Num() == 0);
-				CHECK(Set.Max() == MaxElements);
+				CHECK(Set.Max() == CalculatedMaxElements);
 
 				for (uint32 FullIndex = 0; FullIndex < MaxElements; ++FullIndex)
 				{
@@ -314,7 +377,7 @@ struct TSetTestHelper
 				CHECK(Set.Num() == MaxElements);
 				Set.Reset();
 				CHECK(Set.Num() == 0);
-				CHECK(Set.Max() == MaxElements);
+				CHECK(Set.Max() == CalculatedMaxElements);
 			}
 		}
 
