Fixed TSortedSet default sorting predicate, to support pointers.

diff --git a/Engine/Source/Runtime/Core/Public/Containers/ContainersFwd.h b/Engine/Source/Runtime/Core/Public/Containers/ContainersFwd.h
--- a/Engine/Source/Runtime/Core/Public/Containers/ContainersFwd.h
+++ b/Engine/Source/Runtime/Core/Public/Containers/ContainersFwd.h
@@ -82,7 +82,7 @@ template<typename KeyType, typename ValueType, typename SetAllocator = FDefaultS
 template <typename T = void > struct TLess;
 template <typename> struct TTypeTraits;
 template<typename InKeyType, typename InValueType, typename ArrayAllocator = FDefaultAllocator, typename SortPredicate = TLess<typename TTypeTraits<InKeyType>::ConstPointerType> > class TSortedMap;
-template <typename InElementType, typename ArrayAllocator = FDefaultAllocator, typename SortPredicate = TLess<InElementType>> class TSortedSet;
+template <typename InElementType, typename ArrayAllocator = FDefaultAllocator, typename SortPredicate = TLess<typename TTypeTraits<InElementType>::ConstPointerType>> class TSortedSet;
 template<typename InElementType, typename KeyFuncs = DefaultKeyFuncs<InElementType>, typename Allocator = FDefaultSparseSetAllocator> class TSparseSet;
 template<typename InElementType, typename KeyFuncs = DefaultKeyFuncs<InElementType>, typename Allocator = FDefaultCompactSetAllocator> class TCompactSet;
 template<typename InElementType, typename KeyFuncs = DefaultKeyFuncs<InElementType>, typename Allocator = FDefaultSetAllocator> class TSet;
