CsvProfiler - Change frame time to be computed after UpdateTimeAndHandleMaxTickRate() instead of EndFrame.

When using a framerate limit this fixes the severe noise (+/- 1-3ms) in the csv frametime stat. This is because when framerate limits are enabled there's an idle at the start of the frame (for sake of input latency) based on the duration of the previous frames. Because it cannot possibly know how long the current frame will take this introduces a lot of variance in frame time if computed at the end of the engine tick. Therefore the logical start of the frame is after this idle point up until the end of the next idle time. This results in a very smooth frame time when running with a frame limit.

When running with uncapped fps or vsync enabled this makes no difference in frame times besides them being phase shifted back 1 frame. Note that even despite the phase shift thread stat timings largely line up as they're computed mid frame.

#Synced-CL 51327058

#Synced-CL 51327058

diff --git a/Engine/Source/Runtime/Core/Private/ProfilingDebugging/CsvProfiler.cpp b/Engine/Source/Runtime/Core/Private/ProfilingDebugging/CsvProfiler.cpp
--- a/Engine/Source/Runtime/Core/Private/ProfilingDebugging/CsvProfiler.cpp
+++ b/Engine/Source/Runtime/Core/Private/ProfilingDebugging/CsvProfiler.cpp
@@ -147,6 +147,14 @@ TAutoConsoleVariable<int32> CVarCsvTargetFrameRateOverride(
 	ECVF_Default
 );
 
+TAutoConsoleVariable<bool> CVarCsvUseLegacyFrameTime(
+	TEXT("csv.UseLegacyFrameTime"),
+	false,
+	TEXT("If true uses the delta between EndFrame() calls to compute the frame time.")
+	TEXT("If false, computes the frame time after the idle at the start of the frame (when using a frame rate limit)."),
+	ECVF_Default
+);
+
 #if UE_BUILD_SHIPPING
 TAutoConsoleVariable<int32> CVarCsvShippingContinuousWrites(
 	TEXT("csv.Shipping.ContinuousWrites"),
@@ -4143,11 +4151,15 @@ void FCsvProfiler::EndFrame()
 			}
 		}
 
-		// Record the frametime (measured since the last EndFrame)
-		uint64 CurrentTimeStamp = FPlatformTime::Cycles64();
-		uint64 ElapsedCycles = CurrentTimeStamp - LastEndFrameTimestamp;
-		float ElapsedMs = (float)FPlatformTime::ToMilliseconds64(ElapsedCycles);
-		CSV_CUSTOM_STAT_MINIMAL_GLOBAL(FrameTime, ElapsedMs, ECsvCustomStatOp::Set);
+		const uint64 CurrentTimeStamp = FPlatformTime::Cycles64();
+
+		if (CVarCsvUseLegacyFrameTime->GetBool())
+		{
+			// Record the frametime (measured since the last EndFrame)
+			uint64 ElapsedCycles = CurrentTimeStamp - LastEndFrameTimestamp;
+			double ElapsedMs = FPlatformTime::ToMilliseconds64(ElapsedCycles);
+			CSV_CUSTOM_STAT_MINIMAL_GLOBAL(FrameTime, ElapsedMs, ECsvCustomStatOp::Set);
+		}
 
 		FPlatformMemoryStats MemoryStats = PlatformMemoryHelpers::GetFrameMemoryStats();
 
@@ -4180,6 +4192,8 @@ void FCsvProfiler::EndFrame()
 			ProcessStatData();
 		}
 
+		// Note: we use the end frame timestamp even when CVarCsvUseLegacyFrameTime is false as 
+		// this timestamp is the csv frame boundary and not necessarily the logical frame boundary.
 		LastEndFrameTimestamp = CurrentTimeStamp;
 		CaptureFrameNumber++;
 	}
@@ -4345,6 +4359,30 @@ void FCsvProfiler::BeginFrameEOP()
 	}
 }
 
+void FCsvProfiler::UpdateFrameTime()
+{
+	if (!CVarCsvUseLegacyFrameTime->GetBool())
+	{
+		// This is called right after UpdateTimeAndHandleMaxTickRate() in the engine loop which is the point at which
+		// we've finished waiting when a frame rate limit is enabled. Computing the frame time here means it does not 
+		// align with the csv frame boundary, but it does align with the logical frame boundary.
+		// This avoids noise in the frame time stat introduced by the wait at the start of the frame as it cannot possibly
+		// predict how long this frame will be and is instead based on the time of previous frames. Therefore this wait 
+		// period and any time before it logically belong to the prior frame. Using this period for the frame time results
+		// in almost 0 noise in the frame time calculation when a frame rate limit is enabled, and no difference to computing
+		// it at the end of the frame when running uncapped or with vsync enabled (besides the phase shift).
+
+		static double LastTimestamp = 0.0;
+		const double CurrentTimestamp = FPlatformTime::Seconds();
+		if (LastTimestamp > 0.0)
+		{
+			double FrameTimeMs = (CurrentTimestamp - LastTimestamp) * 1000.0;
+			CSV_CUSTOM_STAT_MINIMAL_GLOBAL(FrameTime, FrameTimeMs, ECsvCustomStatOp::Set);
+		}
+		LastTimestamp = CurrentTimestamp;
+	}
+}
+
 void FCsvProfiler::BeginCapture(int InNumFramesToCapture, 
 	const FString& InDestinationFolder, 
 	const FString& InFilename,
diff --git a/Engine/Source/Runtime/Core/Public/ProfilingDebugging/CsvProfiler.h b/Engine/Source/Runtime/Core/Public/ProfilingDebugging/CsvProfiler.h
--- a/Engine/Source/Runtime/Core/Public/ProfilingDebugging/CsvProfiler.h
+++ b/Engine/Source/Runtime/Core/Public/ProfilingDebugging/CsvProfiler.h
@@ -537,6 +537,9 @@ public:
 
 	CORE_API void BeginFrameEOP();
 
+	/** Called from the engine tick at the point where we should compute the current frame time. */
+	CORE_API void UpdateFrameTime();
+
 	CORE_API void SetDeviceProfileName(FString InDeviceProfileName);
 
 	FString GetOutputFilename() const { return OutputFilename; }
diff --git a/Engine/Source/Runtime/Launch/Private/LaunchEngineLoop.cpp b/Engine/Source/Runtime/Launch/Private/LaunchEngineLoop.cpp
--- a/Engine/Source/Runtime/Launch/Private/LaunchEngineLoop.cpp
+++ b/Engine/Source/Runtime/Launch/Private/LaunchEngineLoop.cpp
@@ -5688,6 +5688,9 @@ void FEngineLoop::Tick()
 			SCOPED_NAMED_EVENT(FEngineLoop_UpdateTimeAndHandleMaxTickRate, SystemColor);
 			GEngine->UpdateTimeAndHandleMaxTickRate();
 			GEngine->SetSimulationLatencyMarkerStart(CurrentFrameCounter);
+#if WITH_ENGINE && CSV_PROFILER
+			FCsvProfiler::Get()->UpdateFrameTime();
+#endif
 		}
 
 		// beginning of RHI frame
