Fix static analysis error

#Synced-CL 51398680

diff --git a/Engine/Source/Runtime/Core/Public/Containers/CompactSetBase.h b/Engine/Source/Runtime/Core/Public/Containers/CompactSetBase.h
--- a/Engine/Source/Runtime/Core/Public/Containers/CompactSetBase.h
+++ b/Engine/Source/Runtime/Core/Public/Containers/CompactSetBase.h
@@ -242,7 +242,7 @@ protected:
 				uint8* NewHashLocation = NewHashTable + ((size_t)NewMaxElements << NewTypeShift) + HashCountSize;
 
 				// Copy hash lookup table (first so to free up the space for the other data to move)
-				FMemory::Memmove(NewHashLocation, OldHashLocation, (NewHashCount << NewTypeShift));
+				FMemory::Memmove(NewHashLocation, OldHashLocation, ((size_t)NewHashCount << NewTypeShift));
 
 				// Copy hash size + next index table
 				FMemory::Memmove(NewHashTable, OldHashTable, ((size_t)NumElements << NewTypeShift) + HashCountSize);
