Added LexFromString to convert a hexadecimal or long package name string to FPackageId.

#Synced-CL 51343971

diff --git a/Engine/Source/Runtime/Core/Private/IO/PackageId.cpp b/Engine/Source/Runtime/Core/Private/IO/PackageId.cpp
--- a/Engine/Source/Runtime/Core/Private/IO/PackageId.cpp
+++ b/Engine/Source/Runtime/Core/Private/IO/PackageId.cpp
@@ -2,14 +2,15 @@
 
 #include "IO/PackageId.h"
 
+#include "HAL/LowLevelMemTracker.h"
 #include "Hash/CityHash.h"
 #include "Misc/Char.h"
-#include "HAL/LowLevelMemTracker.h"
+#include "Misc/Parse.h"
+#include "Misc/ScopeRWLock.h"
+#include "Misc/TransactionallySafeRWLock.h"
 #include "Serialization/CompactBinaryWriter.h"
 #include "Serialization/StructuredArchiveAdapters.h"
 #include "Serialization/StructuredArchiveSlots.h"
-#include "Misc/ScopeRWLock.h"
-#include "Misc/TransactionallySafeRWLock.h"
 
 #if WITH_PACKAGEID_NAME_MAP
 LLM_DEFINE_TAG(PackageId_ReverseMapping);
@@ -93,3 +94,39 @@ FString LexToString(const FPackageId& PackageId)
 	return FString::Printf(TEXT("%llX"), PackageId.Value());
 }
 
+void LexFromString(FPackageId& Out, FStringView String)
+{
+	if (String.IsEmpty())
+	{
+		Out = FPackageId();
+		return;
+	}
+
+	if (String.StartsWith(TEXT("0x"), ESearchCase::IgnoreCase))
+	{
+		String.RightChopInline(2);
+	}
+
+	if (String.Len() == sizeof(uint64) * 2)
+	{
+		Out = FPackageId::FromValue(FParse::HexNumber64(String));
+	}
+	else
+	{
+		// Assume string is a long package name
+		Out = FPackageId::FromName(FName(String));
+	}
+}
+
+void LexFromString(FPackageId& Out, const TCHAR* String)
+{
+	if (String != nullptr)
+	{
+		LexFromString(Out, FStringView(String));
+	}
+	else
+	{
+		Out = FPackageId();
+	}
+}
+
diff --git a/Engine/Source/Runtime/Core/Public/IO/PackageId.h b/Engine/Source/Runtime/Core/Public/IO/PackageId.h
--- a/Engine/Source/Runtime/Core/Public/IO/PackageId.h
+++ b/Engine/Source/Runtime/Core/Public/IO/PackageId.h
@@ -80,6 +80,8 @@ public:
 };
 
 CORE_API FString LexToString(const FPackageId& PackageId);
+CORE_API void LexFromString(FPackageId& Out, FStringView String);
+CORE_API void LexFromString(FPackageId& Out, const TCHAR* String);
 
 template <typename CharType>
 TStringBuilderBase<CharType>& operator<<(TStringBuilderBase<CharType>& Builder, const FPackageId& PackageId)
