Made TOptional<T> trivially copyable when T is trivially copyable.

diff --git a/Engine/Source/Runtime/Core/Public/Misc/IntrusiveUnsetOptionalState.h b/Engine/Source/Runtime/Core/Public/Misc/IntrusiveUnsetOptionalState.h
--- a/Engine/Source/Runtime/Core/Public/Misc/IntrusiveUnsetOptionalState.h
+++ b/Engine/Source/Runtime/Core/Public/Misc/IntrusiveUnsetOptionalState.h
@@ -5,11 +5,13 @@
 /**
  * Facilities to allow types to have an intrusive invalid state which can act as
  * TOptional's 'unset' state, saving space.  A class in such a state will only ever be
- * compared against FIntrusiveUnsetOptionalState or destructed.
+ * compared against FIntrusiveUnsetOptionalState or destructed.  It will be trivially
+ * copied if it's trivially copyable, otherwise its copy constructor and assignment operator
+ * won't be called in this state.
  *
  * A class should implement a constructor taking FIntrusiveUnsetOptionalState, and an 
  * equality comparison operator against FIntrusiveUnsetOptionalState, which will put a
-  class instance into the 'unset' state (in the case of the constructor) and allow testing
+ * class instance into the 'unset' state (in the case of the constructor) and allow testing
  * of its unset state.
  *
  * A public constexpr static data member of type bool called bHasIntrusiveUnsetOptionalState
diff --git a/Engine/Source/Runtime/Core/Public/Misc/Optional.h b/Engine/Source/Runtime/Core/Public/Misc/Optional.h
--- a/Engine/Source/Runtime/Core/Public/Misc/Optional.h
+++ b/Engine/Source/Runtime/Core/Public/Misc/Optional.h
@@ -235,7 +235,7 @@ public:
 #endif
 
 	/** Copy/Move construction */
-	[[nodiscard]] TOptional(const TOptional& Other)
+	[[nodiscard]] TOptional(const TOptional& Other) requires (!std::is_trivially_copy_constructible_v<OptionalType>)
 		: TOptional()
 	{
 		bool bLocalIsSet = Other.IsSet();
@@ -249,7 +249,7 @@ public:
 		}
 		// Default constructor initialized unset intrusive optional if necessary
 	}
-	[[nodiscard]] TOptional(TOptional&& Other)
+	[[nodiscard]] TOptional(TOptional&& Other) requires (!std::is_trivially_move_constructible_v<OptionalType>)
 		: TOptional()
 	{
 		bool bLocalIsSet = Other.IsSet();
@@ -264,7 +264,7 @@ public:
 		// Default constructor initialized unset intrusive optional if necessary
 	}
 
-	TOptional& operator=(const TOptional& Other)
+	TOptional& operator=(const TOptional& Other) requires (!std::is_trivially_copy_assignable_v<OptionalType>)
 	{
 		if (&Other != this)
 		{
@@ -279,7 +279,7 @@ public:
 		}
 		return *this;
 	}
-	TOptional& operator=(TOptional&& Other)
+	TOptional& operator=(TOptional&& Other) requires (!std::is_trivially_move_assignable_v<OptionalType>)
 	{
 		if (&Other != this)
 		{
@@ -295,6 +295,11 @@ public:
 		return *this;
 	}
 
+	[[nodiscard]] TOptional(const TOptional& Other) requires (std::is_trivially_copy_constructible_v<OptionalType>) = default;
+	[[nodiscard]] TOptional(TOptional&& Other) requires (std::is_trivially_move_constructible_v<OptionalType>) = default;
+	TOptional& operator=(const TOptional& Other) requires (std::is_trivially_copy_assignable_v<OptionalType>) = default;
+	TOptional& operator=(TOptional&& Other) requires (std::is_trivially_move_assignable_v<OptionalType>) = default;
+
 	TOptional& operator=(const OptionalType& InValue)
 	{
 		if (std::addressof(InValue) != std::addressof(this->TypedValue))
diff --git a/Engine/Source/Runtime/Core/Tests/Misc/OptionalTest.cpp b/Engine/Source/Runtime/Core/Tests/Misc/OptionalTest.cpp
--- a/Engine/Source/Runtime/Core/Tests/Misc/OptionalTest.cpp
+++ b/Engine/Source/Runtime/Core/Tests/Misc/OptionalTest.cpp
@@ -5,6 +5,7 @@
 #if WITH_TESTS
 
 #include "CoreTypes.h"
+#include "Containers/StaticArray.h"
 #include "Containers/UnrealString.h"
 #include "Internationalization/Text.h"
 #include "Math/Vector.h"
@@ -144,4 +145,17 @@ TEST_CASE_NAMED(FTOptionalTest_Intrusive, "System::Core::TOptional::Intrusive",
 	CHECK_MESSAGE("Unset optional is not destroyed", FIntrusiveDestructible::DestructorCount == 0);
 }
 
+// Test that trivial copyability is transitive
+static_assert(std::is_trivially_copyable_v<TOptional<bool>>                      == std::is_trivially_copyable_v<bool>);
+static_assert(std::is_trivially_copyable_v<TOptional<int32>>                     == std::is_trivially_copyable_v<int32>);
+static_assert(std::is_trivially_copyable_v<TOptional<double>>                    == std::is_trivially_copyable_v<double>);
+static_assert(std::is_trivially_copyable_v<TOptional<FName>>                     == std::is_trivially_copyable_v<FName>);
+static_assert(std::is_trivially_copyable_v<TOptional<FString*>>                  == std::is_trivially_copyable_v<FString*>);
+static_assert(std::is_trivially_copyable_v<TOptional<TSharedPtr<int32>>>         == std::is_trivially_copyable_v<TSharedPtr<int32>>);
+static_assert(std::is_trivially_copyable_v<TOptional<FString>>                   == std::is_trivially_copyable_v<FString>);
+static_assert(std::is_trivially_copyable_v<TOptional<FVector>>                   == std::is_trivially_copyable_v<FVector>);
+static_assert(std::is_trivially_copyable_v<TOptional<TArray<int32>>>             == std::is_trivially_copyable_v<TArray<int32>>);
+static_assert(std::is_trivially_copyable_v<TOptional<TStaticArray<int32, 20>>>   == std::is_trivially_copyable_v<TStaticArray<int32, 20>>);
+static_assert(std::is_trivially_copyable_v<TOptional<TStaticArray<FString, 20>>> == std::is_trivially_copyable_v<TStaticArray<FString, 20>>);
+
 #endif // WITH_TESTS
\ No newline at end of file
