Added EAllowShrinking parameters to TSortedMap and TSortedSet removal functions.

diff --git a/Engine/Source/Runtime/Core/Public/Containers/SortedMap.h b/Engine/Source/Runtime/Core/Public/Containers/SortedMap.h
--- a/Engine/Source/Runtime/Core/Public/Containers/SortedMap.h
+++ b/Engine/Source/Runtime/Core/Public/Containers/SortedMap.h
@@ -244,7 +244,7 @@ public:
 	 * @param InKey The key to remove associated values for.
 	 * @return The number of values that were associated with the key.
 	 */
-	inline int32 Remove(KeyConstPointerType InKey)
+	inline int32 Remove(KeyConstPointerType InKey, EAllowShrinking AllowShrinking = UE::Core::Private::AllowShrinkingByDefault<ArrayAllocator>())
 	{
 		int32 RemoveIndex = FindIndex(InKey);
 
@@ -253,7 +253,7 @@ public:
 			return 0;
 		}
 
-		Pairs.RemoveAt(RemoveIndex);
+		Pairs.RemoveAt(RemoveIndex, AllowShrinking);
 
 		return 1;
 	}
@@ -450,7 +450,7 @@ public:
 	 * @param OutRemovedValue If found, the value that was removed (not modified if the key was not found)
 	 * @return Whether or not the key was found
 	 */
-	inline bool RemoveAndCopyValue(KeyInitType Key, ValueType& OutRemovedValue)
+	inline bool RemoveAndCopyValue(KeyInitType Key, ValueType& OutRemovedValue, EAllowShrinking AllowShrinking = UE::Core::Private::AllowShrinkingByDefault<ArrayAllocator>())
 	{
 		int32 FoundIndex = FindIndex(Key);
 		if (FoundIndex == INDEX_NONE)
@@ -459,7 +459,7 @@ public:
 		}
 
 		OutRemovedValue = MoveTemp(Pairs[FoundIndex].Value);
-		Pairs.RemoveAt(FoundIndex);
+		Pairs.RemoveAt(FoundIndex, AllowShrinking);
 		return true;
 	}
 	
@@ -470,13 +470,13 @@ public:
 	 * @param Key The key to search for
 	 * @return Whether or not the key was found
 	 */
-	inline ValueType FindAndRemoveChecked(KeyConstPointerType Key)
+	inline ValueType FindAndRemoveChecked(KeyConstPointerType Key, EAllowShrinking AllowShrinking = UE::Core::Private::AllowShrinkingByDefault<ArrayAllocator>())
 	{
 		int32 FoundIndex = FindIndex(Key);
 		check(FoundIndex != INDEX_NONE);
 		
 		ValueType OutRemovedValue = MoveTemp(Pairs[FoundIndex].Value);
-		Pairs.RemoveAt(FoundIndex);
+		Pairs.RemoveAt(FoundIndex, AllowShrinking);
 		return OutRemovedValue;
 	}
 
diff --git a/Engine/Source/Runtime/Core/Public/Containers/SortedSet.h b/Engine/Source/Runtime/Core/Public/Containers/SortedSet.h
--- a/Engine/Source/Runtime/Core/Public/Containers/SortedSet.h
+++ b/Engine/Source/Runtime/Core/Public/Containers/SortedSet.h
@@ -226,7 +226,7 @@ public:
 	 * @param	InKey	The key to match elements against.
 	 * @return	The number of elements removed.
 	 */
-	UE_FORCEINLINE_HINT int32 Remove(KeyInitType InKey)
+	UE_FORCEINLINE_HINT int32 Remove(KeyInitType InKey, EAllowShrinking AllowShrinking = UE::Core::Private::AllowShrinkingByDefault<ArrayAllocator>())
 	{
 		int32 RemoveIndex = FindIndex(InKey);
 		if (RemoveIndex == INDEX_NONE)
@@ -234,7 +234,7 @@ public:
 			return 0;
 		}
 
-		Elements.RemoveAt(RemoveIndex);
+		Elements.RemoveAt(RemoveIndex, AllowShrinking);
 		return 1;
 	}
 
