[Insights] Add support for controlling max size when writing the tail buffer

#Synced-CL 51335054

diff --git a/Engine/Source/Runtime/Core/Private/ProfilingDebugging/TraceAuxiliary.cpp b/Engine/Source/Runtime/Core/Private/ProfilingDebugging/TraceAuxiliary.cpp
--- a/Engine/Source/Runtime/Core/Private/ProfilingDebugging/TraceAuxiliary.cpp
+++ b/Engine/Source/Runtime/Core/Private/ProfilingDebugging/TraceAuxiliary.cpp
@@ -175,8 +175,8 @@ public:
 	void StartWorkerThread();
 	void StartEndFrameUpdate();
 	void RegisterEndFrameCallbacks();
-	bool WriteSnapshot(const TCHAR* InFilePath, const FTraceAuxiliary::FLogCategoryAlias& LogCategory);
-	bool SendSnapshot(const TCHAR* InHost, uint32 InPort, const FTraceAuxiliary::FLogCategoryAlias& LogCategory);
+	bool WriteSnapshot(const TCHAR* InFilePath, const FTraceAuxiliary::FLogCategoryAlias& LogCategory, uint32 MaxSize);
+	bool SendSnapshot(const TCHAR* InHost, uint32 InPort, const FTraceAuxiliary::FLogCategoryAlias& LogCategory, uint32 MaxSize);
 	void AddCommandlineChannels(const TCHAR* ChannelList);
 	void ResetCommandlineChannels();
 	void UpdateStats();
@@ -1183,7 +1183,7 @@ bool FTraceAuxiliaryImpl::WriteToFile(const TCHAR* Path, const FTraceAuxiliary::
 }
 
 ////////////////////////////////////////////////////////////////////////////////
-bool FTraceAuxiliaryImpl::WriteSnapshot(const TCHAR* InFilePath, const FTraceAuxiliary::FLogCategoryAlias& LogCategory)
+bool FTraceAuxiliaryImpl::WriteSnapshot(const TCHAR* InFilePath, const FTraceAuxiliary::FLogCategoryAlias& LogCategory, uint32 MaxTailSize)
 {
 	double StartTime = FPlatformTime::Seconds();
 
@@ -1195,7 +1195,7 @@ bool FTraceAuxiliaryImpl::WriteSnapshot(const TCHAR* InFilePath, const FTraceAux
 
 	UE_LOG_REF(LogCategory, Log, TEXT("Writing trace snapshot to '%s'..."), *NativePath);
 
-	const bool bResult = UE::Trace::WriteSnapshotTo(*NativePath);
+	const bool bResult = UE::Trace::WriteSnapshotTo(*NativePath, MaxTailSize);
 
 	if (bResult)
 	{
@@ -1212,7 +1212,7 @@ bool FTraceAuxiliaryImpl::WriteSnapshot(const TCHAR* InFilePath, const FTraceAux
 }
 
 ////////////////////////////////////////////////////////////////////////////////
-bool FTraceAuxiliaryImpl::SendSnapshot(const TCHAR* InHost, uint32 InPort, const FTraceAuxiliary::FLogCategoryAlias& LogCategory)
+bool FTraceAuxiliaryImpl::SendSnapshot(const TCHAR* InHost, uint32 InPort, const FTraceAuxiliary::FLogCategoryAlias& LogCategory, uint32 MaxTailSize)
 {
 	double StartTime = FPlatformTime::Seconds();
 
@@ -1224,7 +1224,7 @@ bool FTraceAuxiliaryImpl::SendSnapshot(const TCHAR* InHost, uint32 InPort, const
 
 	UE_LOG_REF(LogCategory, Log, TEXT("Sending trace snapshot to '%s'..."), InHost);
 
-	const bool bResult = UE::Trace::SendSnapshotTo(InHost, InPort);
+	const bool bResult = UE::Trace::SendSnapshotTo(InHost, InPort, MaxTailSize);
 
 	if (bResult)
 	{
@@ -1693,7 +1693,7 @@ static void TraceAuxiliarySnapshotFile(const TArray<FString>& Args)
 		return;
 	}
 
-	GTraceAuxiliary.WriteSnapshot(FilePath, LogConsoleResponse);
+	GTraceAuxiliary.WriteSnapshot(FilePath, LogConsoleResponse, 0);
 }
 
 ////////////////////////////////////////////////////////////////////////////////
@@ -1716,7 +1716,7 @@ static void TraceAuxiliarySnapshotSend(const TArray<FString>& Args)
 		return;
 	}
 
-	GTraceAuxiliary.SendSnapshot(Host, Port, LogConsoleResponse);
+	GTraceAuxiliary.SendSnapshot(Host, Port, LogConsoleResponse, 0);
 }
 
 ////////////////////////////////////////////////////////////////////////////////
@@ -2128,22 +2128,22 @@ bool FTraceAuxiliary::Resume()
 }
 
 ////////////////////////////////////////////////////////////////////////////////
-bool FTraceAuxiliary::WriteSnapshot(const TCHAR* InFilePath)
+bool FTraceAuxiliary::WriteSnapshot(const TCHAR* InFilePath, uint32 MaxTailSize)
 {
 #if UE_TRACE_MINIMAL_ENABLED
 	TRACE_CPUPROFILER_EVENT_SCOPE(FTraceAux_WriteSnapshot);
-	return GTraceAuxiliary.WriteSnapshot(InFilePath, LogTrace);
+	return GTraceAuxiliary.WriteSnapshot(InFilePath, LogTrace, MaxTailSize);
 #else
 	return true;
 #endif
 }
 
 ////////////////////////////////////////////////////////////////////////////////
-bool FTraceAuxiliary::SendSnapshot(const TCHAR* Host, uint32 Port)
+bool FTraceAuxiliary::SendSnapshot(const TCHAR* Host, uint32 Port, uint32 MaxTailSize)
 {
 #if UE_TRACE_MINIMAL_ENABLED
 	TRACE_CPUPROFILER_EVENT_SCOPE(FTraceAux_SendSnapshot);
-	return GTraceAuxiliary.SendSnapshot(Host, Port, LogTrace);
+	return GTraceAuxiliary.SendSnapshot(Host, Port, LogTrace, MaxTailSize);
 #else
 	return true;
 #endif
diff --git a/Engine/Source/Runtime/Core/Public/ProfilingDebugging/TraceAuxiliary.h b/Engine/Source/Runtime/Core/Public/ProfilingDebugging/TraceAuxiliary.h
--- a/Engine/Source/Runtime/Core/Public/ProfilingDebugging/TraceAuxiliary.h
+++ b/Engine/Source/Runtime/Core/Public/ProfilingDebugging/TraceAuxiliary.h
@@ -247,14 +247,19 @@ public:
 	/**
 	 * Write tailing memory state to a utrace file.
 	 * @param FilePath Path to the file to write the snapshot to. If it is null or empty a file path will be generated.
+	 * @param MaxTailSize Maximum size of tail buffer to write. Zero means no limit. Total file size may be larger due to important events and
+	 * overhead.
 	 */
-	static CORE_API bool WriteSnapshot(const TCHAR* FilePath);
+	static CORE_API bool WriteSnapshot(const TCHAR* FilePath, uint32 MaxTailSize = 0);
 
 	/**
 	 * Write tailing memory state to a trace server.
-	 * @param FilePath Path to the file to write the snapshot to. If it is null or empty a file path will be generated.
+	 * @param Host Host to send snapshot to
+	 * @param Port Port to connect to host on. Zero means default port is used 
+	 * @param MaxTailSize Maximum size of tail buffer to write. Zero means no limit. Total stream size may be larger due to important events and
+	 * overhead.
 	 */
-	static CORE_API bool SendSnapshot(const TCHAR* Host = nullptr, uint32 Port = 0);
+	static CORE_API bool SendSnapshot(const TCHAR* Host = nullptr, uint32 Port = 0, uint32 MaxTailSize = 0);
 
 	/**
 	 * Initialize Trace systems.
diff --git a/Engine/Source/Runtime/TraceLog/Private/Trace/Tail.cpp b/Engine/Source/Runtime/TraceLog/Private/Trace/Tail.cpp
--- a/Engine/Source/Runtime/TraceLog/Private/Trace/Tail.cpp
+++ b/Engine/Source/Runtime/TraceLog/Private/Trace/Tail.cpp
@@ -53,7 +53,7 @@ public:
 	FRange			GetBackPackets() const;
 	FRange			GetFrontPackets() const;
 	template <typename CallbackType>
-	void			IterateRanges(CallbackType&& Callback);
+	void			IterateRanges(CallbackType&& Callback, uint32 MaxSize);
 	template <typename PacketType>
 	PacketType*		Append(uint32 InSize);
 	template <typename PacketType>
@@ -127,19 +127,37 @@ FPacketRing::FRange FPacketRing::GetFrontPackets() const
 
 ////////////////////////////////////////////////////////////////////////////////
 template <typename CallbackType>
-void FPacketRing::IterateRanges(CallbackType&& Callback)
+void FPacketRing::IterateRanges(CallbackType&& Callback, uint32 MaxSize)
 {
 	FPacketRing::FRange Ranges[] = { GetBackPackets(), GetFrontPackets() };
 
-	// Send out the ranges.
+	int32 BytesToSkip = 0;
+	if (MaxSize != 0 && MaxSize < INT32_MAX)
+	{
+		const uint32 TotalSize = Ranges[0].Size + Ranges[1].Size;
+		const uint32 TotalBytesToSkip = MaxSize > TotalSize ? 0 : TotalSize - MaxSize;
+		BytesToSkip = TotalBytesToSkip;
+	}
+
 	for (const auto& Range : Ranges)
 	{
 		if (Range.Size == 0)
 		{
 			continue;
 		}
-
-		Callback(Range);
+		uint32 RangeSkipped = 0;
+		uint8* RangeOffset = (uint8*) Range.Data;
+		while (BytesToSkip > 0 && RangeOffset < ((uint8*)Range.Data + Range.Size))
+		{
+			// We need to skip over entire packages
+			FTidPacketBase* Packet = (FTidPacketBase*) RangeOffset;
+			BytesToSkip -= Packet->PacketSize;
+			RangeSkipped += Packet->PacketSize;
+			RangeOffset = RangeOffset + Packet->PacketSize;
+		} 
+
+		FPacketRing::FRange RangeView = {RangeOffset, Range.Size - RangeSkipped};
+		Callback(RangeView);
 	}
 }
 
@@ -295,7 +313,7 @@ void Writer_TailAppend(uint32 ThreadId, uint8* __restrict Data, uint32 Size)
 }
 
 ////////////////////////////////////////////////////////////////////////////////
-void Writer_TailOnConnect()
+void Writer_TailOnConnect(uint32 MaxBytes /*= 0*/)
 {
 	FProfilerScope _(__func__);
 	
@@ -308,7 +326,7 @@ void Writer_TailOnConnect()
 	GPacketRing.IterateRanges([] (const FPacketRing::FRange& Range)
 	{
 		Writer_SendDataRaw(Range.Data, Range.Size);
-	});
+	}, MaxBytes);
 }
 
 ////////////////////////////////////////////////////////////////////////////////
diff --git a/Engine/Source/Runtime/TraceLog/Private/Trace/Trace.cpp b/Engine/Source/Runtime/TraceLog/Private/Trace/Trace.cpp
--- a/Engine/Source/Runtime/TraceLog/Private/Trace/Trace.cpp
+++ b/Engine/Source/Runtime/TraceLog/Private/Trace/Trace.cpp
@@ -27,8 +27,8 @@ void	Writer_Update();
 bool	Writer_SendTo(const ANSICHAR*, uint32, uint32);
 bool	Writer_WriteTo(const ANSICHAR*, uint32);
 bool	Writer_RelayTo(UPTRINT, IoWriteFunc, IoCloseFunc, uint16);
-bool	Writer_WriteSnapshotTo(const ANSICHAR*);
-bool	Writer_SendSnapshotTo(const ANSICHAR*, uint32);
+bool	Writer_WriteSnapshotTo(const ANSICHAR*, uint32);
+bool	Writer_SendSnapshotTo(const ANSICHAR*, uint32, uint32);
 bool	Writer_IsTracing();
 bool	Writer_IsTracingTo(uint32 (&OutSessionGuid)[4], uint32 (&OutTraceGuid)[4]);
 bool	Writer_Stop();
@@ -163,19 +163,19 @@ bool RelayTo(UPTRINT InHandle, IoWriteFunc WriteFunc, IoCloseFunc CloseFunc, uin
 }
 
 ////////////////////////////////////////////////////////////////////////////////
-bool WriteSnapshotTo(const TCHAR* InPath)
+bool WriteSnapshotTo(const TCHAR* InPath, uint32 MaxTailSize)
 {
 	char Path[512];
 	ToAnsiCheap(Path, InPath);
-	return Private::Writer_WriteSnapshotTo(Path);
+	return Private::Writer_WriteSnapshotTo(Path, MaxTailSize);
 }
 
 ////////////////////////////////////////////////////////////////////////////////
-bool SendSnapshotTo(const TCHAR* InHost, uint32 InPort)
+bool SendSnapshotTo(const TCHAR* InHost, uint32 InPort, uint32 MaxTailSize)
 {
 	char Host[512];
 	ToAnsiCheap(Host, InHost);
-	return Private::Writer_SendSnapshotTo(Host, InPort);
+	return Private::Writer_SendSnapshotTo(Host, InPort, MaxTailSize);
 }
 
 ////////////////////////////////////////////////////////////////////////////////
diff --git a/Engine/Source/Runtime/TraceLog/Private/Trace/Writer.cpp b/Engine/Source/Runtime/TraceLog/Private/Trace/Writer.cpp
--- a/Engine/Source/Runtime/TraceLog/Private/Trace/Writer.cpp
+++ b/Engine/Source/Runtime/TraceLog/Private/Trace/Writer.cpp
@@ -44,7 +44,7 @@ void			Writer_SetUpdateCallback(OnUpdateFunc* Callback);
 void			Writer_SendData(uint32, uint8* __restrict, uint32);
 void			Writer_InitializeTail(int32);
 void			Writer_ShutdownTail();
-void			Writer_TailOnConnect();
+void			Writer_TailOnConnect(uint32 MaxSize = 0);
 void			Writer_InitializeSharedBuffers();
 void			Writer_ShutdownSharedBuffers();
 void			Writer_UpdateSharedBuffers();
@@ -1405,6 +1405,7 @@ struct FSnapshotTarget
 {
 	enum EType { FileTarget, HostTarget };
 	EType Type;
+	uint32 MaxTailSize;
 	union
 	{
 		 struct
@@ -1512,7 +1513,7 @@ bool Writer_WriteSnapshot(const FSnapshotTarget& Target)
 		Writer_CallbackOnConnect();
 
 		// Finally write the events in the tail buffer
-		Writer_TailOnConnect();
+		Writer_TailOnConnect(Target.MaxTailSize);
 
 		// Send sync packets to help parsers digest any out-of-order events
 		GSyncPacketCountdown = GNumSyncPackets;
@@ -1528,19 +1529,21 @@ bool Writer_WriteSnapshot(const FSnapshotTarget& Target)
 }
 
 ////////////////////////////////////////////////////////////////////////////////
-bool Writer_WriteSnapshotTo(const ANSICHAR* Path)
+bool Writer_WriteSnapshotTo(const ANSICHAR* Path, uint32 MaxTailSize)
 {
 	FSnapshotTarget Target;
 	Target.Type = FSnapshotTarget::EType::FileTarget;
+	Target.MaxTailSize = MaxTailSize;
 	Target.File.Path = Path;
 	return Writer_WriteSnapshot(Target);
 }
 
 ////////////////////////////////////////////////////////////////////////////////
-bool Writer_SendSnapshotTo(const ANSICHAR* Host, uint32 Port)
+bool Writer_SendSnapshotTo(const ANSICHAR* Host, uint32 Port, uint32 MaxTailSize)
 {
 	FSnapshotTarget Target;
 	Target.Type = FSnapshotTarget::EType::HostTarget;
+	Target.MaxTailSize = MaxTailSize;
 	Target.Host.Host = Host;
 	Target.Host.Port = Port;
 	return Writer_WriteSnapshot(Target);
diff --git a/Engine/Source/Runtime/TraceLog/Public/Trace/Trace.h b/Engine/Source/Runtime/TraceLog/Public/Trace/Trace.h
--- a/Engine/Source/Runtime/TraceLog/Public/Trace/Trace.h
+++ b/Engine/Source/Runtime/TraceLog/Public/Trace/Trace.h
@@ -350,15 +350,17 @@ UE_TRACE_API bool	RelayTo(UPTRINT InHandle, IoWriteFunc WriteFunc, IoCloseFunc C
 /**
  * Immediately write contents of tail buffers and important events to a new or existing file.
  * @param Path Target path
+ * @param MaxTailSize Maximum size of tail buffer to write. Zero indicates no limit.
  */
-UE_TRACE_API bool	WriteSnapshotTo(const TCHAR* Path) UE_TRACE_IMPL(false);
+UE_TRACE_API bool	WriteSnapshotTo(const TCHAR* Path, uint32 MaxTailSize = 0) UE_TRACE_IMPL(false);
 
 /**
  * Immediately write contents of tail buffers and important events to a remote host.
  * @param Host Target hostname or ip
  * @param Port Target port
+ * @param MaxTailSize Maximum size of tail buffer to write. Zero indicates no limit.
  */
-UE_TRACE_API bool	SendSnapshotTo(const TCHAR* Host, uint32 Port) UE_TRACE_IMPL(false);
+UE_TRACE_API bool	SendSnapshotTo(const TCHAR* Host, uint32 Port, uint32 MaxTailSize = 0) UE_TRACE_IMPL(false);
 
 /**
  * Checks if TraceLog currently has an output. Note that trace events can still be recorded
