Added && TShared/Weak... constructors and StaticCastShared/Weak... helpers. Removed WeakPtr from TSharedFromThis

#Synced-CL 51289644

diff --git a/Engine/Source/Runtime/Core/Public/Templates/SharedPointer.h b/Engine/Source/Runtime/Core/Public/Templates/SharedPointer.h
--- a/Engine/Source/Runtime/Core/Public/Templates/SharedPointer.h
+++ b/Engine/Source/Runtime/Core/Public/Templates/SharedPointer.h
@@ -229,7 +229,7 @@ public:
 
 		// If the object happens to be derived from TSharedFromThis, the following method
 		// will prime the object with a weak pointer to itself.
-		SharedPointerInternals::EnableSharedFromThis( this, InRawPtrProxy.Object );
+		SharedPointerInternals::EnableSharedFromThis(SharedReferenceCount, InRawPtrProxy.Object);
 	}
 
 	/**
@@ -279,7 +279,7 @@ public:
 
 		// If the object happens to be derived from TSharedFromThis, the following method
 		// will prime the object with a weak pointer to itself.
-		SharedPointerInternals::EnableSharedFromThis( this, InRawPtrProxy.Object );
+		SharedPointerInternals::EnableSharedFromThis(SharedReferenceCount, InRawPtrProxy.Object);
 	}
 
 	/**
@@ -588,7 +588,7 @@ private:
 
 		// If the object happens to be derived from TSharedFromThis, the following method
 		// will prime the object with a weak pointer to itself.
-		SharedPointerInternals::EnableSharedFromThis(this, InObject);
+		SharedPointerInternals::EnableSharedFromThis(SharedReferenceCount, InObject);
 	}
 
 	/**
@@ -725,7 +725,7 @@ public:
 	{
 		// If the object happens to be derived from TSharedFromThis, the following method
 		// will prime the object with a weak pointer to itself.
-		SharedPointerInternals::EnableSharedFromThis( this, InObject );
+		SharedPointerInternals::EnableSharedFromThis(SharedReferenceCount, InObject);
 	}
 
 	/**
@@ -746,7 +746,7 @@ public:
 	{
 		// If the object happens to be derived from TSharedFromThis, the following method
 		// will prime the object with a weak pointer to itself.
-		SharedPointerInternals::EnableSharedFromThis( this, InObject );
+		SharedPointerInternals::EnableSharedFromThis(SharedReferenceCount, InObject);
 	}
 
 	/**
@@ -765,7 +765,7 @@ public:
 	{
 		// If the object happens to be derived from TSharedFromThis, the following method
 		// will prime the object with a weak pointer to itself.
-		SharedPointerInternals::EnableSharedFromThis( this, InRawPtrProxy.Object );
+		SharedPointerInternals::EnableSharedFromThis(SharedReferenceCount, InRawPtrProxy.Object);
 	}
 
 	/**
@@ -785,7 +785,7 @@ public:
 	{
 		// If the object happens to be derived from TSharedFromThis, the following method
 		// will prime the object with a weak pointer to itself.
-		SharedPointerInternals::EnableSharedFromThis( this, InRawPtrProxy.Object );
+		SharedPointerInternals::EnableSharedFromThis(SharedReferenceCount, InRawPtrProxy.Object);
 	}
 
 	/**
@@ -805,7 +805,7 @@ public:
 	{
 		// If the object happens to be derived from TSharedFromThis, the following method
 		// will prime the object with a weak pointer to itself.
-		SharedPointerInternals::EnableSharedFromThis( this, InRawPtrProxy.Object );
+		SharedPointerInternals::EnableSharedFromThis(SharedReferenceCount, InRawPtrProxy.Object);
 	}
 
 	/**
@@ -871,6 +871,14 @@ public:
 	{
 	}
 	
+	template <typename OtherType>
+	UE_NODEBUG TSharedPtr(TSharedPtr<OtherType, Mode>&& InSharedPtr, SharedPointerInternals::FStaticCastTag)
+		: Object{static_cast<ObjectType*>(InSharedPtr.Object)}
+		, SharedReferenceCount{MoveTemp(InSharedPtr.SharedReferenceCount)}
+	{
+		InSharedPtr.Object = nullptr;
+	}
+	
 	/**
 	 * Special constructor used internally to cast a 'const' shared pointer a 'mutable' pointer.  You
 	 * should never call this constructor directly.  Instead, use the ConstCastSharedPtr() function.
@@ -886,6 +894,14 @@ public:
 	{
 	}
 
+	template <typename OtherType>
+	UE_NODEBUG TSharedPtr(TSharedPtr<OtherType, Mode>&& InSharedPtr, SharedPointerInternals::FConstCastTag)
+		: Object{const_cast<ObjectType*>(InSharedPtr.Object)}
+		, SharedReferenceCount{MoveTemp(InSharedPtr.SharedReferenceCount)}
+	{
+		InSharedPtr.Object = nullptr;
+	}
+	
 	/**
 	 * Aliasing constructor used to create a shared pointer which shares its reference count with
 	 * another shared object, but pointing to a different object, typically a subobject.
@@ -1243,6 +1259,19 @@ private:
 		}
 	}
 
+	/**
+	 * Aliasing constructor used to create a shared pointer which shares its reference count with
+	 * another shared object, but pointing to a different object, typically a subobject.
+	 *
+	 * @param  InSharedReferenceCount  The shared referencer whose reference count should be shared.
+	 * @param  InObject                The object pointer to use
+	 */
+	UE_NODEBUG TSharedPtr(SharedPointerInternals::FSharedReferencer<Mode> InSharedReferenceCount, ObjectType* InObject)
+		: Object{InObject}
+		, SharedReferenceCount{MoveTemp(InSharedReferenceCount)}
+	{
+	}
+
 	// We declare ourselves as a friend (templated using OtherType) so we can access members as needed
 	template< class OtherType, ESPMode OtherMode > friend class TSharedPtr;
 
@@ -1356,6 +1385,14 @@ public:
 	{
 	}
 
+	template <typename OtherType>
+	UE_NODEBUG TWeakPtr(TWeakPtr<OtherType, Mode>&& InWeakPtr, SharedPointerInternals::FStaticCastTag)
+		: Object{static_cast<ObjectType*>(InWeakPtr.Object)}
+		, WeakReferenceCount{MoveTemp(InWeakPtr.WeakReferenceCount)}
+	{
+		InWeakPtr.Object = nullptr;
+	}
+
 	/**
 	 * Special constructor used internally to cast a 'const' weak pointer a 'mutable' pointer.  You
 	 * should never call this constructor directly.  Instead, use the ConstCastWeakPtr() function.
@@ -1371,6 +1408,14 @@ public:
 	{
 	}
 
+	template <typename OtherType>
+	UE_NODEBUG TWeakPtr(TWeakPtr<OtherType, Mode>&& InWeakPtr, SharedPointerInternals::FConstCastTag)
+		: Object{const_cast<ObjectType*>(InWeakPtr.Object)}
+		, WeakReferenceCount{MoveTemp(InWeakPtr.WeakReferenceCount)}
+	{
+		InWeakPtr.Object = nullptr;
+	}
+
 	/**
 	 * Constructs a weak pointer from a weak pointer of another type.
 	 * This constructor is intended to allow derived-to-base conversions.
@@ -1603,11 +1648,25 @@ public:
 	}
 
 private:
+	/**
+	 * Aliasing constructor used to create a weak pointer which shares its reference count with
+	 * another shared object, but pointing to a different object, typically a subobject.
+	 *
+	 * @param  InWeakReferenceCount  The weak referencer whose reference count should be shared.
+	 * @param  InObject              The object pointer to use
+	 */
+	UE_NODEBUG TWeakPtr(SharedPointerInternals::FWeakReferencer<Mode> InWeakReferenceCount, ObjectType* InObject)
+		: Object{InObject}
+		, WeakReferenceCount{MoveTemp(InWeakReferenceCount)}
+	{
+	}
+
 	// We declare ourselves as a friend (templated using OtherType) so we can access members as needed
     template< class OtherType, ESPMode OtherMode > friend class TWeakPtr;
 
 	// Declare ourselves as a friend of TSharedPtr so we can access members as needed
     template< class OtherType, ESPMode OtherMode > friend class TSharedPtr;
+	template <class OtherType, ESPMode OtherMode> friend class TSharedFromThis;
 
 private:
 	/** The object we have a weak reference to.  Can be nullptr.  Also, it's important to note that because
@@ -1657,8 +1716,6 @@ public:
 	 */
 	[[nodiscard]] UE_NODEBUG TSharedRef< ObjectType, Mode > AsShared()
 	{
-		TSharedPtr< ObjectType, Mode > SharedThis( WeakThis.Pin() );
-
 		//
 		// If the following assert goes off, it means one of the following:
 		//
@@ -1668,7 +1725,10 @@ public:
 		// To fix this, make sure you create at least one shared reference to your object instance before requested,
 		// and also avoid calling this function from your object's destructor.
 		//
-		check( SharedThis.Get() == this );
+		check(DoesSharedInstanceExist());
+		
+		ReferenceController->AddSharedReference();
+		TSharedPtr<ObjectType, Mode> SharedThis{SharedPointerInternals::FSharedReferencer<Mode>{ReferenceController}, static_cast<ObjectType*>(this)};
 
 		// Now that we've verified the shared pointer is valid, we'll convert it to a shared reference
 		// and return it!
@@ -1684,8 +1744,6 @@ public:
 	 */
 	[[nodiscard]] UE_NODEBUG TSharedRef< ObjectType const, Mode > AsShared() const
 	{
-		TSharedPtr< ObjectType const, Mode > SharedThis( WeakThis );
-
 		//
 		// If the following assert goes off, it means one of the following:
 		//
@@ -1695,7 +1753,10 @@ public:
 		// To fix this, make sure you create at least one shared reference to your object instance before requested,
 		// and also avoid calling this function from your object's destructor.
 		//
-		check( SharedThis.Get() == this );
+		check(DoesSharedInstanceExist());
+		
+		ReferenceController->AddSharedReference();
+		TSharedPtr<ObjectType const, Mode> SharedThis{SharedPointerInternals::FSharedReferencer<Mode>{ReferenceController}, static_cast<ObjectType const*>(this)};
 
 		// Now that we've verified the shared pointer is valid, we'll convert it to a shared reference
 		// and return it!
@@ -1725,8 +1786,6 @@ public:
 	 */
 	[[nodiscard]] UE_NODEBUG TWeakPtr< ObjectType, Mode > AsWeak()
 	{
-		TWeakPtr< ObjectType, Mode > Result = WeakThis;
-
 		//
 		// If the following assert goes off, it means one of the following:
 		//
@@ -1736,15 +1795,16 @@ public:
 		// To fix this, make sure you create at least one shared reference to your object instance before requested,
 		// and also avoid calling this function from your object's destructor.
 		//
-		check( Result.Pin().Get() == this );
+		check(DoesSharedInstanceExist());
+		
+		ReferenceController->AddWeakReference();
+		TWeakPtr<ObjectType, Mode> Result{SharedPointerInternals::FWeakReferencer<Mode>{ReferenceController}, static_cast<ObjectType*>(this)};
 
 		// Now that we've verified the pointer is valid, we'll return it!
 		return Result;
 	}
 	[[nodiscard]] UE_NODEBUG TWeakPtr< ObjectType const, Mode > AsWeak() const
 	{
-		TWeakPtr< ObjectType const, Mode > Result = WeakThis;
-
 		//
 		// If the following assert goes off, it means one of the following:
 		//
@@ -1754,7 +1814,10 @@ public:
 		// To fix this, make sure you create at least one shared reference to your object instance before requested,
 		// and also avoid calling this function from your object's destructor.
 		//
-		check( Result.Pin().Get() == this );
+		check(DoesSharedInstanceExist());
+		
+		ReferenceController->AddWeakReference();
+		TWeakPtr<ObjectType const, Mode> Result{SharedPointerInternals::FWeakReferencer<Mode>{ReferenceController}, static_cast<ObjectType const*>(this)};
 
 		// Now that we've verified the pointer is valid, we'll return it!
 		return Result;
@@ -1804,39 +1867,24 @@ protected:
 		return StaticCastSharedRef< OtherType const >( ThisPtr->AsShared() );
 	}
 
-public:		// @todo: Ideally this would be private, but template sharing problems prevent it
+private:
 
 	/**
-	 * INTERNAL USE ONLY -- Do not call this method.  Freshens the internal weak pointer object using
-	 * the supplied object pointer along with the authoritative shared reference to the object.
+	 * Freshens the internal weak pointer object using the authoritative shared reference to the object.
 	 * Note that until this function is called, calls to AsShared() will result in an empty pointer.
 	 */
-	template< class SharedPtrType, ESPMode SharedPtrMode, class OtherType >
-	UE_NODEBUG inline void UpdateWeakReferenceInternal( TSharedPtr< SharedPtrType, SharedPtrMode > const* InSharedPtr, OtherType* InObject ) const
+	template <ESPMode OtherMode>
+	UE_NODEBUG void UpdateWeakReferenceInternal(SharedPointerInternals::TReferenceControllerBase<OtherMode>* InReferenceController) const
 	{
-		static_assert(SharedPtrMode == Mode, "You cannot use a TSharedPtr of one mode with a type which inherits TSharedFromThis of another mode.");
+		static_assert(OtherMode == Mode, "You cannot create a shared reference or pointer of one mode with a type which inherits TSharedFromThis of another mode.");
 
-		if( !WeakThis.IsValid() )
+		if (!ReferenceController)
 		{
-			WeakThis = TSharedPtr< ObjectType, Mode >( *InSharedPtr, InObject );
+			ReferenceController = InReferenceController;
 		}
 	}
 
-	/**
-	 * INTERNAL USE ONLY -- Do not call this method.  Freshens the internal weak pointer object using
-	 * the supplied object pointer along with the authoritative shared reference to the object.
-	 * Note that until this function is called, calls to AsShared() will result in an empty pointer.
-	 */
-	template< class SharedRefType, ESPMode SharedPtrMode, class OtherType >
-	UE_NODEBUG inline void UpdateWeakReferenceInternal( TSharedRef< SharedRefType, SharedPtrMode > const* InSharedRef, OtherType* InObject ) const
-	{
-		static_assert(SharedPtrMode == Mode, "You cannot use a TSharedPtr of one mode with a type which inherits TSharedFromThis of another mode.");
-
-		if( !WeakThis.IsValid() )
-		{
-			WeakThis = TSharedRef< ObjectType, Mode >( *InSharedRef, InObject );
-		}
-	}
+public:
 
 	/**
 	 * Checks whether our referenced instance is valid (ie, whether it's safe to call AsShared).
@@ -1847,7 +1895,7 @@ public:		// @todo: Ideally this would be private, but template sharing problems
 	 */
 	[[nodiscard]] UE_REWRITE bool DoesSharedInstanceExist() const
 	{
-		return WeakThis.IsValid();
+		return ReferenceController && ReferenceController->GetSharedReferenceCount() > 0;
 	}
 
 protected:
@@ -1876,10 +1924,12 @@ protected:
 	~TSharedFromThis() = default;
 
 private:
+	
+	template <class SharedReferencerType, class OtherType>
+	friend void SharedPointerInternals::EnableSharedFromThis(SharedReferencerType& InSharedReferencer, OtherType const* InObject);
 
-	/** Weak reference to ourselves.  If we're destroyed then this weak pointer reference will be destructed
-	    with ourselves.  Note this is declared mutable only so that UpdateWeakReferenceInternal() can update it. */
-	mutable TWeakPtr< ObjectType, Mode > WeakThis;	
+	/** Note this is declared mutable only so that UpdateWeakReferenceInternal() can update it. */
+	mutable SharedPointerInternals::TReferenceControllerBase<Mode>* ReferenceController = nullptr;
 };
 
 
@@ -1894,6 +1944,12 @@ template< class CastToType, class CastFromType, ESPMode Mode >
 	return TSharedPtr< CastToType, Mode >(InSharedPtr, SharedPointerInternals::FStaticCastTag());
 }
 
+template <class CastToType, class CastFromType, ESPMode Mode>
+[[nodiscard]] UE_REWRITE TSharedPtr<CastToType, Mode> StaticCastSharedPtr(TSharedPtr<CastFromType, Mode>&& InSharedPtr)
+{
+	return TSharedPtr<CastToType, Mode>{MoveTemp(InSharedPtr), SharedPointerInternals::FStaticCastTag{}};
+}
+
 
 /**
  * Casts a weak pointer of one type to another type. (static_cast)  Useful for down-casting.
@@ -1906,6 +1962,12 @@ template< class CastToType, class CastFromType, ESPMode Mode >
 	return TWeakPtr< CastToType, Mode >(InWeakPtr, SharedPointerInternals::FStaticCastTag());
 }
 
+template <class CastToType, class CastFromType, ESPMode Mode>
+[[nodiscard]] UE_REWRITE TWeakPtr<CastToType, Mode> StaticCastWeakPtr(TWeakPtr<CastFromType, Mode>&& InWeakPtr)
+{
+	return TWeakPtr<CastToType, Mode>{MoveTemp(InWeakPtr), SharedPointerInternals::FStaticCastTag{}};
+}
+
 
 /**
  * Casts a 'const' shared reference to 'mutable' shared reference. (const_cast)
@@ -1918,6 +1980,12 @@ template< class CastToType, class CastFromType, ESPMode Mode >
 	return TSharedRef< CastToType, Mode >( InSharedRef, SharedPointerInternals::FConstCastTag() );
 }
 
+template <class CastToType, class CastFromType, ESPMode Mode>
+[[nodiscard]] UE_REWRITE TSharedRef<CastToType, Mode> ConstCastSharedRef(TSharedRef<CastFromType, Mode>&& InSharedRef)
+{
+	return TSharedRef<CastToType, Mode>{MoveTemp(InSharedRef), SharedPointerInternals::FConstCastTag{}};
+}
+
 
 /**
  * Casts a 'const' shared pointer to 'mutable' shared pointer. (const_cast)
@@ -1930,6 +1998,12 @@ template< class CastToType, class CastFromType, ESPMode Mode >
 	return TSharedPtr< CastToType, Mode >( InSharedPtr, SharedPointerInternals::FConstCastTag() );
 }
 
+template <class CastToType, class CastFromType, ESPMode Mode>
+[[nodiscard]] UE_REWRITE TSharedPtr<CastToType, Mode> ConstCastSharedPtr(TSharedPtr<CastFromType, Mode>&& InSharedPtr)
+{
+	return TSharedPtr<CastToType, Mode>{MoveTemp(InSharedPtr), SharedPointerInternals::FConstCastTag{}};
+}
+
 
 /**
  * Casts a 'const' weak pointer to 'mutable' weak pointer. (const_cast)
@@ -1942,6 +2016,12 @@ template< class CastToType, class CastFromType, ESPMode Mode >
 	return TWeakPtr< CastToType, Mode >(InWeakPtr, SharedPointerInternals::FConstCastTag());
 }
 
+template <class CastToType, class CastFromType, ESPMode Mode>
+[[nodiscard]] UE_REWRITE TWeakPtr<CastToType, Mode> ConstCastWeakPtr(TWeakPtr<CastFromType, Mode>&& InWeakPtr)
+{
+	return TWeakPtr<CastToType, Mode>{MoveTemp(InWeakPtr), SharedPointerInternals::FConstCastTag{}};
+}
+
 
 /**
  * MakeShareable utility function.  Wrap object pointers with MakeShareable to allow them to be implicitly
diff --git a/Engine/Source/Runtime/Core/Public/Templates/SharedPointerInternals.h b/Engine/Source/Runtime/Core/Public/Templates/SharedPointerInternals.h
--- a/Engine/Source/Runtime/Core/Public/Templates/SharedPointerInternals.h
+++ b/Engine/Source/Runtime/Core/Public/Templates/SharedPointerInternals.h
@@ -33,6 +33,9 @@ namespace SharedPointerInternals
 	// Forward declarations
 	template< ESPMode Mode > class FWeakReferencer;
 
+	template <class SharedReferencerType, class ObjectType>
+	void EnableSharedFromThis(SharedReferencerType& InSharedReferencer, ObjectType const* InObject);
+
 	/** Dummy structures used internally as template arguments for typecasts */
 	struct FStaticCastTag {};
 	struct FConstCastTag {};
@@ -538,7 +541,7 @@ namespace SharedPointerInternals
 			: ReferenceController( nullptr )
 		{ }
 
-		/** Constructor that counts a single reference to the specified object */
+		/** Constructor that references the specified controller, without incrementing the shared count */
 		inline explicit FSharedReferencer( TReferenceControllerBase<Mode>* InReferenceController )
 			: ReferenceController( InReferenceController )
 		{ }
@@ -697,6 +700,9 @@ namespace SharedPointerInternals
  		// Expose access to ReferenceController to FWeakReferencer
 		template< ESPMode OtherMode > friend class FWeakReferencer;
 
+		template <class SharedReferencerType, class OtherType>
+		friend void EnableSharedFromThis(SharedReferencerType& InSharedReferencer, OtherType const* InObject);
+
 	private:
 
 		/** Pointer to the reference controller for the object a shared reference/pointer is referencing */
@@ -718,6 +724,12 @@ namespace SharedPointerInternals
 			: ReferenceController( nullptr )
 		{ }
 
+		/** Constructor that references the specified controller, without incrementing the weak count */
+		explicit FWeakReferencer(TReferenceControllerBase<Mode>* InReferenceController)
+			: ReferenceController{InReferenceController}
+		{
+		}
+
 		/** Construct a weak referencer object from another weak referencer */
 		inline FWeakReferencer( FWeakReferencer const& InWeakRefCountPointer )
 			: ReferenceController( InWeakRefCountPointer.ReferenceController )
@@ -857,9 +869,9 @@ constexpr bool IsDerivedFromSharedFromThis()
 
 namespace SharedPointerInternals
 {
-	/** Templated helper function (const) that creates a shared pointer from an object instance */
-	template< class SharedPtrType, class ObjectType >
-	inline void EnableSharedFromThis( SharedPtrType* InSharedPtrOrRef, ObjectType const* InObject )
+	/** Templated helper function (const) that creates a shared reference or pointer from an object instance */
+	template <class SharedReferencerType, class ObjectType>
+	void EnableSharedFromThis(SharedReferencerType& InSharedReferencer, ObjectType const* InObject)
 	{
 		///////////////////////////////////////////////////////////////////////////////////////////////
 		// If you get an 'ambiguous call' compile error in this function, it means you have multiple //
@@ -868,9 +880,9 @@ namespace SharedPointerInternals
 
 		if constexpr (IsDerivedFromSharedFromThis<ObjectType>())
 		{
-			if( InObject != nullptr )
+			if (InObject)
 			{
-				InObject->UpdateWeakReferenceInternal( InSharedPtrOrRef, const_cast< ObjectType* >( InObject ) );
+				InObject->UpdateWeakReferenceInternal(InSharedReferencer.ReferenceController);
 			}
 		}
 	}
