Add missing jni

#Synced-CL 51468882

diff --git a/Engine/Source/Runtime/Core/Public/Android/AndroidJniGameActivity.h b/Engine/Source/Runtime/Core/Public/Android/AndroidJniGameActivity.h
--- a/Engine/Source/Runtime/Core/Public/Android/AndroidJniGameActivity.h
+++ b/Engine/Source/Runtime/Core/Public/Android/AndroidJniGameActivity.h
@@ -144,9 +144,11 @@ namespace UE::Jni::Asis
 			UE_JNI_NATIVE_METHOD(nativeSetCommandline),
 			UE_JNI_NATIVE_METHOD(nativeAppCommand)
 		};
+#if USE_ANDROID_INPUT
 
 	private:
 		using FGameActivity::Members;
+#endif
 	};
 }
 #if USE_ANDROID_STANDALONE
