Remove debug NSLog that is triggering ASan error detection

#rnx

diff --git a/Engine/Source/Runtime/Core/Private/IOS/IOSPlatformMisc.cpp b/Engine/Source/Runtime/Core/Private/IOS/IOSPlatformMisc.cpp
--- a/Engine/Source/Runtime/Core/Private/IOS/IOSPlatformMisc.cpp
+++ b/Engine/Source/Runtime/Core/Private/IOS/IOSPlatformMisc.cpp
@@ -1779,7 +1779,6 @@ static FEntitlementsData TryFindEntitlementsData()
 				uint32 dataLength = ntohl(blob.length) - sizeof(uint32_t) * 2;
 				char* data = (char*)malloc(dataLength);
 				fread(&data[0], sizeof(char) * dataLength, 1, file);
-				NSLog(@"%s", data);
 				fclose(file);
 				return FEntitlementsData(data, dataLength);
 			}
