Added an option to toggle mb3 large allocations cache on\off in run time

#Synced-CL 51337004

diff --git a/Engine/Source/Runtime/Core/Private/HAL/MallocBinned3.cpp b/Engine/Source/Runtime/Core/Private/HAL/MallocBinned3.cpp
--- a/Engine/Source/Runtime/Core/Private/HAL/MallocBinned3.cpp
+++ b/Engine/Source/Runtime/Core/Private/HAL/MallocBinned3.cpp
@@ -18,6 +18,15 @@ PRAGMA_DISABLE_UNSAFE_TYPECAST_WARNINGS
 #	define UE_MB3_MAX_CACHED_OS_FREES (64)
 #	define UE_MB3_MAX_CACHED_OS_FREES_BYTE_LIMIT (64*1024*1024)
 
+	enum class EMB3CacheState : int8
+	{
+		Enabled = 0,
+		Disabled = 1,
+		Trimmed = 2,
+	};
+
+	std::atomic<EMB3CacheState> LargeAllocsCacheEnabledState(EMB3CacheState::Enabled);
+
 	typedef TCachedOSPageAllocator<UE_MB3_MAX_CACHED_OS_FREES, UE_MB3_MAX_CACHED_OS_FREES_BYTE_LIMIT> TBinned3CachedOSPageAllocator;
 
 	TBinned3CachedOSPageAllocator& GetCachedOSPageAllocator()
@@ -25,7 +34,7 @@ PRAGMA_DISABLE_UNSAFE_TYPECAST_WARNINGS
 		static TBinned3CachedOSPageAllocator Singleton;
 		return Singleton;
 	}
-#endif
+#endif	// ~UE_MB3_USE_CACHED_PAGE_ALLOCATOR_FOR_LARGE_ALLOCS
 
 #if UE_MB3_ALLOCATOR_STATS
 	TAtomic<int64> Binned3Commits;
@@ -699,21 +708,28 @@ void* FMallocBinned3::Malloc(SIZE_T Size, uint32 Alignment)
 #endif
 
 	void* Result;
-#if !UE_MB3_USE_CACHED_PAGE_ALLOCATOR_FOR_LARGE_ALLOCS
-	FPlatformMemory::FPlatformVirtualMemoryBlock Block = FPlatformMemory::FPlatformVirtualMemoryBlock::AllocateVirtual(AlignedSize, Alignment);
-	Block.Commit(0, AlignedSize);
-	Result = Block.GetVirtualPointer();
-	LLM_IF_ENABLED(FLowLevelMemTracker::Get().OnLowLevelAlloc(ELLMTracker::Platform, Result, AlignedSize));
-	const size_t ActualSize = Block.GetActualSizeInPages();
-#endif
-
 	{
-		UE::TUniqueLock Lock(ExternalAllocMutex);
+		size_t VMPages;
 
 #if UE_MB3_USE_CACHED_PAGE_ALLOCATOR_FOR_LARGE_ALLOCS
-		Result = GetCachedOSPageAllocator().Allocate(AlignedSize);
-		const size_t ActualSize = AlignedSize / FPlatformMemory::FPlatformVirtualMemoryBlock::GetCommitAlignment();
+		if (LargeAllocsCacheEnabledState.load(std::memory_order_relaxed) == EMB3CacheState::Enabled)
+		{
+			VMPages = AlignedSize / FPlatformMemory::FPlatformVirtualMemoryBlock::GetCommitAlignment();
+
+			UE::TUniqueLock Lock(ExternalAllocMutex);
+			Result = GetCachedOSPageAllocator().Allocate(AlignedSize);
+		}
+		else
 #endif
+		{
+			FPlatformMemory::FPlatformVirtualMemoryBlock Block = FPlatformMemory::FPlatformVirtualMemoryBlock::AllocateVirtual(AlignedSize, Alignment);
+			Block.Commit(0, AlignedSize);
+			Result = Block.GetVirtualPointer();
+			LLM_IF_ENABLED(FLowLevelMemTracker::Get().OnLowLevelAlloc(ELLMTracker::Platform, Result, AlignedSize));
+			VMPages = Block.GetActualSizeInPages();
+		}
+
+		UE::TUniqueLock Lock(ExternalAllocMutex);
 		if (!Result)
 		{
 			OutOfMemory(AlignedSize);
@@ -722,7 +738,7 @@ void* FMallocBinned3::Malloc(SIZE_T Size, uint32 Alignment)
 		// Create pool.
 		FPoolInfo* Pool = Internal::GetOrCreatePoolInfo(*this, Result, FPoolInfo::ECanary::Assigned);
 		check(Size > 0 && Size <= AlignedSize && AlignedSize >= FPlatformMemory::FPlatformVirtualMemoryBlock::GetCommitAlignment());
-		Pool->SetOSAllocationSizes(Size, AlignedSize, ActualSize);
+		Pool->SetOSAllocationSizes(Size, AlignedSize, VMPages);
 	}
 
 #if UE_MB3_TIME_LARGE_BLOCKS
@@ -887,16 +903,21 @@ void FMallocBinned3::FreeExternal(void* Ptr, uint64 PoolIndex)
 #if UE_MB3_TIME_LARGE_BLOCKS
 		const double StartTime = FPlatformTime::Seconds();
 #endif
-		{
+
 #if UE_MB3_USE_CACHED_PAGE_ALLOCATOR_FOR_LARGE_ALLOCS
+		if (LargeAllocsCacheEnabledState.load(std::memory_order_relaxed) == EMB3CacheState::Enabled)
+		{
 			UE::TUniqueLock Lock(ExternalAllocMutex);
 			GetCachedOSPageAllocator().Free(Ptr, VMPages * FPlatformMemory::FPlatformVirtualMemoryBlock::GetCommitAlignment());
-#else
+		}
+		else
+#endif
+		{
 			LLM_IF_ENABLED(FLowLevelMemTracker::Get().OnLowLevelFree(ELLMTracker::Platform, Ptr));
 			FPlatformMemory::FPlatformVirtualMemoryBlock Block(Ptr, VMPages);
 			Block.FreeVirtual();
-#endif
 		}
+
 #if UE_MB3_TIME_LARGE_BLOCKS
 		const double Add = FPlatformTime::Seconds() - StartTime;
 		double Old;
@@ -928,6 +949,16 @@ void FMallocBinned3::Trim(bool bTrimThreadCaches)
 		// Trim memory and increase the Epoch.
 		FMallocBinnedCommonUtils::Trim(*this);
 	}
+
+#if UE_MB3_USE_CACHED_PAGE_ALLOCATOR_FOR_LARGE_ALLOCS
+	if (LargeAllocsCacheEnabledState.load(std::memory_order_relaxed) == EMB3CacheState::Disabled)
+	{
+		// trim cache as some allocations might have been stored in it due to a race condition
+		UE::TUniqueLock Lock(ExternalAllocMutex);
+		GetCachedOSPageAllocator().FreeAll();
+		LargeAllocsCacheEnabledState.store(EMB3CacheState::Trimmed, std::memory_order_relaxed);
+	}
+#endif
 }
 
 void FMallocBinned3::FlushCurrentThreadCacheInternal(bool bNewEpochOnly)
@@ -935,6 +966,22 @@ void FMallocBinned3::FlushCurrentThreadCacheInternal(bool bNewEpochOnly)
 	FMallocBinnedCommonUtils::FlushCurrentThreadCache(*this, bNewEpochOnly);
 }
 
+void FMallocBinned3::EnableLargeAllocsCache()
+{
+#if UE_MB3_USE_CACHED_PAGE_ALLOCATOR_FOR_LARGE_ALLOCS
+	LargeAllocsCacheEnabledState.store(EMB3CacheState::Enabled, std::memory_order_relaxed);
+#endif
+}
+
+void FMallocBinned3::DisableLargeAllocsCache()
+{
+#if UE_MB3_USE_CACHED_PAGE_ALLOCATOR_FOR_LARGE_ALLOCS
+	UE::TUniqueLock Lock(ExternalAllocMutex);
+	LargeAllocsCacheEnabledState.store(EMB3CacheState::Disabled, std::memory_order_relaxed);
+	GetCachedOSPageAllocator().FreeAll();
+#endif
+}
+
 void FMallocBinned3::DumpAllocatorStats(class FOutputDevice& Ar)
 {
 #if UE_MB3_ALLOCATOR_STATS
diff --git a/Engine/Source/Runtime/Core/Public/HAL/MallocBinned3.h b/Engine/Source/Runtime/Core/Public/HAL/MallocBinned3.h
--- a/Engine/Source/Runtime/Core/Public/HAL/MallocBinned3.h
+++ b/Engine/Source/Runtime/Core/Public/HAL/MallocBinned3.h
@@ -301,6 +301,9 @@ public:
 
 	void FlushCurrentThreadCacheInternal(bool bNewEpochOnly = false);
 
+	void EnableLargeAllocsCache();
+	void DisableLargeAllocsCache();
+
 	static void* AllocateMetaDataMemory(SIZE_T Size);
 	static void FreeMetaDataMemory(void *Ptr, SIZE_T Size);
 };
