Deprecated PLATFORM_COMPILER_SUPPORTS_CONSTRAINED_DESTRUCTORS.

diff --git a/Engine/Source/Runtime/Core/Public/HAL/Platform.h b/Engine/Source/Runtime/Core/Public/HAL/Platform.h
--- a/Engine/Source/Runtime/Core/Public/HAL/Platform.h
+++ b/Engine/Source/Runtime/Core/Public/HAL/Platform.h
@@ -281,16 +281,7 @@
 #ifndef PLATFORM_COMPILER_SUPPORTS_BUILTIN_BITCAST
 	#error PLATFORM_COMPILER_SUPPORTS_BUILTIN_BITCAST should have been defined by now
 #endif
-#ifndef PLATFORM_COMPILER_SUPPORTS_CONSTRAINED_DESTRUCTORS
-	// Disable support for constrained destructors on MSVC until this internal compile error gets fixed:
-	//
-	// https://developercommunity.visualstudio.com/t/Internal-compile-error-when-returning-a/10650542
-	#if defined(__cpp_concepts) && __cpp_concepts >= 202002L && (defined(__clang__) || defined(__GNUC__) || !defined(_MSC_VER) || _MSC_VER >= 1942)
-		#define PLATFORM_COMPILER_SUPPORTS_CONSTRAINED_DESTRUCTORS 1
-	#else
-		#define PLATFORM_COMPILER_SUPPORTS_CONSTRAINED_DESTRUCTORS 0
-	#endif
-#endif
+#define PLATFORM_COMPILER_SUPPORTS_CONSTRAINED_DESTRUCTORS UE_DEPRECATED_MACRO(5.8, "PLATFORM_COMPILER_SUPPORTS_CONSTRAINED_DESTRUCTORS has been deprecated and should be replaced with 1.") 1
 #define PLATFORM_TCHAR_IS_4_BYTES UE_DEPRECATED_MACRO(5.8, "PLATFORM_TCHAR_IS_4_BYTES has been deprecated and should be replaced with 0.") 0
 #ifndef PLATFORM_WCHAR_IS_4_BYTES
 	#define PLATFORM_WCHAR_IS_4_BYTES			0
diff --git a/Engine/Source/Runtime/Core/Public/Misc/Optional.h b/Engine/Source/Runtime/Core/Public/Misc/Optional.h
--- a/Engine/Source/Runtime/Core/Public/Misc/Optional.h
+++ b/Engine/Source/Runtime/Core/Public/Misc/Optional.h
@@ -37,87 +37,6 @@ namespace UE::Core::Private
 			}
 		}
 	};
-
-#if !PLATFORM_COMPILER_SUPPORTS_CONSTRAINED_DESTRUCTORS
-	// These base types are necessary for trivial destruction on compilers that don't yet support use of
-	// constraints to select an explicitly defaulted trivial destructor.
-	// Once such compilers are no longer supported, all (!PLATFORM_COMPILER_SUPPORTS_CONSTRAINED_DESTRUCTORS)
-	// code in this file can be deleted.
-	template <typename OptionalType, bool bIsTriviallyDestructible = std::is_trivially_destructible_v<OptionalType>>
-	struct TOptionalBase 
-	{
-		union
-		{
-			OptionalType TypedValue;
-		};
-		// Set flag takes up no space if bIsUsingIntrusiveUnsetState
-		static constexpr bool bUsingIntrusiveUnsetState = HasIntrusiveUnsetOptionalState<OptionalType>();
-		UE_NO_UNIQUE_ADDRESS std::conditional_t<bUsingIntrusiveUnsetState, FEmpty, bool> bIsSet = {};
-	
-		constexpr TOptionalBase()
-		{
-		}
-
-		template<typename... ArgTypes>
-		constexpr TOptionalBase(EInPlace, ArgTypes&&... InArgs)
-			: TypedValue(Forward<ArgTypes>(InArgs)...)
-		{
-		}
-
-		// Explicitly constexpr destructor for trivially destructible OptionalType
-		constexpr ~TOptionalBase() = default;
-
-		[[nodiscard]] constexpr bool IsSet() const
-		{
-			return FOptional::IsSet(this);
-		}
-
-		UE_FORCEINLINE_HINT constexpr void DestroyValue()
-		{
-			// Type is known to be trivially destructible
-		}
-	};
-
-	template <typename OptionalType>
-	struct TOptionalBase<OptionalType, false> 
-	{
-		union
-		{
-			OptionalType TypedValue;
-		};
-		// Set flag takes up no space if bIsUsingIntrusiveUnsetState
-		static constexpr bool bUsingIntrusiveUnsetState = HasIntrusiveUnsetOptionalState<OptionalType>();
-		UE_NO_UNIQUE_ADDRESS std::conditional_t<bUsingIntrusiveUnsetState, FEmpty, bool> bIsSet = {};
-
-		constexpr TOptionalBase()
-		{
-		}
-
-		template<typename... ArgTypes>
-		constexpr TOptionalBase(EInPlace, ArgTypes&&... InArgs)
-			: TypedValue(Forward<ArgTypes>(InArgs)...)
-		{
-		}
-
-		constexpr ~TOptionalBase()
-		{
-			if (IsSet())
-			{
-				DestroyValue();
-			}
-		}
-
-		[[nodiscard]] constexpr bool IsSet() const
-		{
-			return FOptional::IsSet(this);
-		}
-
-		UE_FORCEINLINE_HINT constexpr void DestroyValue()
-		{
-			DestructItem(std::addressof(this->TypedValue));
-		}
-	};
-#endif
 }
 
 /**
@@ -126,15 +45,8 @@ namespace UE::Core::Private
  */
 template<typename OptionalType>
 struct TOptional
-#if !PLATFORM_COMPILER_SUPPORTS_CONSTRAINED_DESTRUCTORS
-    : private UE::Core::Private::TOptionalBase<OptionalType>
-#endif
 {
 private:
-#if !PLATFORM_COMPILER_SUPPORTS_CONSTRAINED_DESTRUCTORS
-    using Super = UE::Core::Private::TOptionalBase<OptionalType>;
-#endif
-
 	using FOptional = UE::Core::Private::FOptional;
 	friend FOptional;
 	static constexpr bool bUsingIntrusiveUnsetState = HasIntrusiveUnsetOptionalState<OptionalType>();
@@ -161,11 +73,7 @@ public:
 	}
 	template <typename... ArgTypes>
 	[[nodiscard]] explicit constexpr TOptional(EInPlace, ArgTypes&&... Args)
-#if PLATFORM_COMPILER_SUPPORTS_CONSTRAINED_DESTRUCTORS
 		: TypedValue(Forward<ArgTypes>(Args)...)
-#else
-		: Super(InPlace, Forward<ArgTypes>(Args)...)
-#endif
 	{
 		// If this fails to compile when trying to call TOptional(EInPlace, ...) with a non-public constructor,
 		// do not make TOptional a friend.
@@ -207,11 +115,7 @@ public:
 
 	/** Construct an OptionalType intrusively with no value; i.e. unset */
 	[[nodiscard]] constexpr TOptional() requires bUsingIntrusiveUnsetState
-#if PLATFORM_COMPILER_SUPPORTS_CONSTRAINED_DESTRUCTORS
 		: TypedValue(FIntrusiveUnsetOptionalState{})
-#else
-		: Super(InPlace, FIntrusiveUnsetOptionalState{})
-#endif
 	{
 	}
 
@@ -220,7 +124,6 @@ public:
 	{
 	}
 
-#if PLATFORM_COMPILER_SUPPORTS_CONSTRAINED_DESTRUCTORS
 	constexpr ~TOptional() requires std::is_trivially_destructible_v<OptionalType> = default;
 	constexpr ~TOptional() requires (!std::is_trivially_destructible_v<OptionalType>)
 	{
@@ -230,9 +133,6 @@ public:
 			DestroyValue();
 		}
 	}
-#else
-	constexpr ~TOptional() = default; // Defer to base type to handle trivial/non-trivial destructor
-#endif
 
 	/** Copy/Move construction */
 	[[nodiscard]] TOptional(const TOptional& Other)
@@ -434,14 +334,10 @@ public:
 	}
 
 	/** @return true when the value is meaningful; false if calling GetValue() is undefined. */
-#if PLATFORM_COMPILER_SUPPORTS_CONSTRAINED_DESTRUCTORS
 	[[nodiscard]] constexpr bool IsSet() const
 	{
 		return FOptional::IsSet(this);
 	}
-#else
-	using Super::IsSet;
-#endif
 
 	[[nodiscard]] UE_FORCEINLINE_HINT explicit constexpr operator bool() const
 	{
@@ -505,24 +401,18 @@ private:
 	/** 
 	 * Destroys the value, must only be called if the value is set, and callers must then mark the value unset or construct a new value in place. 
 	 */
-#if PLATFORM_COMPILER_SUPPORTS_CONSTRAINED_DESTRUCTORS
 	UE_FORCEINLINE_HINT constexpr void DestroyValue()
 	{
 		DestructItem(std::addressof(this->TypedValue));
 	}
-#else
-	using Super::DestroyValue;
-#endif
 
 	// If we can have an explicitly conditional trivial destructor, we can declare our members inline and avoid the base class
-#if PLATFORM_COMPILER_SUPPORTS_CONSTRAINED_DESTRUCTORS
 	union
 	{
 		OptionalType TypedValue;
 	};
 	// Set flag takes up no space if bIsUsingIntrusiveUnsetState, otherwise is zero-initialized to false
 	UE_NO_UNIQUE_ADDRESS std::conditional_t<bUsingIntrusiveUnsetState, UE::Core::Private::FEmpty, bool> bIsSet = {};
-#endif
 };
 
 template<typename OptionalType>
diff --git a/Engine/Source/Runtime/Core/Public/Misc/TVariant.h b/Engine/Source/Runtime/Core/Public/Misc/TVariant.h
--- a/Engine/Source/Runtime/Core/Public/Misc/TVariant.h
+++ b/Engine/Source/Runtime/Core/Public/Misc/TVariant.h
@@ -39,18 +39,9 @@ struct FEmptyVariantState
  * Attempting to use the value of a Get() when the underlying type is different leads to undefined behavior.
  */
 template <typename T, typename... Ts>
-class TVariant final
-#if PLATFORM_COMPILER_SUPPORTS_CONSTRAINED_DESTRUCTORS
-	: private UE::Core::Private::TVariantStorage<T, Ts...>
-#else
-	: private std::conditional_t<!std::is_trivially_destructible_v<T> || (!std::is_trivially_destructible_v<Ts> || ...), UE::Core::Private::TDestructibleVariantStorage<T, Ts...> , UE::Core::Private::TVariantStorage<T, Ts...>>
-#endif
+class TVariant final : private UE::Core::Private::TVariantStorage<T, Ts...>
 {
-#if PLATFORM_COMPILER_SUPPORTS_CONSTRAINED_DESTRUCTORS
 	using Super = UE::Core::Private::TVariantStorage<T, Ts...>;
-#else
-	using Super = std::conditional_t<!std::is_trivially_destructible_v<T> || (!std::is_trivially_destructible_v<Ts> || ...), UE::Core::Private::TDestructibleVariantStorage<T, Ts...> , UE::Core::Private::TVariantStorage<T, Ts...>>;
-#endif
 
 	static_assert(!UE::Core::Private::TTypePackContainsDuplicates<T, Ts...>::Value, "All the types used in TVariant should be unique");
 	static_assert(!UE::Core::Private::TContainsReferenceType<T, Ts...>::Value, "TVariant cannot hold reference types");
@@ -118,7 +109,6 @@ public:
 		return *this;
 	}
 
-#if PLATFORM_COMPILER_SUPPORTS_CONSTRAINED_DESTRUCTORS
 	/** Destruct the underlying type (if appropriate) */
 	~TVariant()
 		requires(!std::is_trivially_destructible_v<T> || (!std::is_trivially_destructible_v<Ts> || ...))
@@ -128,10 +118,6 @@ public:
 	~TVariant()
 		requires(std::is_trivially_destructible_v<T> && (std::is_trivially_destructible_v<Ts> && ...))
 	= default;
-#else
-	// Defer to the storage as to how to destruct the elements
-	~TVariant() = default;
-#endif
 
 	/** Determine if the variant holds the specific type */
 	template <typename U>
@@ -254,12 +240,8 @@ public:
 	}
 
 private:
-#if PLATFORM_COMPILER_SUPPORTS_CONSTRAINED_DESTRUCTORS
 	/** Index into the template parameter pack for the type held. */
 	uint8 TypeIndex;
-#else
-	using Super::TypeIndex;
-#endif
 };
 
 /** Apply a visitor function to the list of variants */
diff --git a/Engine/Source/Runtime/Core/Public/Misc/TVariantMeta.h b/Engine/Source/Runtime/Core/Public/Misc/TVariantMeta.h
--- a/Engine/Source/Runtime/Core/Public/Misc/TVariantMeta.h
+++ b/Engine/Source/Runtime/Core/Public/Misc/TVariantMeta.h
@@ -156,11 +156,6 @@ namespace Private
 		}
 
 		TAlignedBytes<SizeofValue, AlignofValue> Storage;
-
-#if !PLATFORM_COMPILER_SUPPORTS_CONSTRAINED_DESTRUCTORS
-		/** Index into the template parameter pack for the type held. */
-		uint8 TypeIndex;
-#endif
 	};
 
 	/** Helper to lookup indices of each type in a template parameter pack */
@@ -382,19 +377,6 @@ namespace Private
 		static constexpr InvokeFn Invokers[] = { WrapperType::template FuncPtr<EncodedIndices>... };
 		return Invokers[EncodedIndex](Forward<Func>(Callable), Forward<Variants>(Args)...);
 	}
-
-#if !PLATFORM_COMPILER_SUPPORTS_CONSTRAINED_DESTRUCTORS
-	template <typename T, typename... Ts>
-	struct TDestructibleVariantStorage : TVariantStorage<T, Ts...>
-	{
-	protected:
-		~TDestructibleVariantStorage()
-		{
-			UE::Core::Private::TDestructorLookup<T, Ts...>::Destruct(this->TypeIndex, this);
-		}
-	};
-#endif
-
 } // namespace Private
 } // namespace Core
 } // namespace UE
