Add range testing methods (TestBetween/TestBetweenEqual) to FAutomationTestBase with improved error messages.

#Synced-CL 51495466
#Synced-CL 51495466

diff --git a/Engine/Source/Runtime/Core/Private/Misc/AutomationTest.cpp b/Engine/Source/Runtime/Core/Private/Misc/AutomationTest.cpp
--- a/Engine/Source/Runtime/Core/Private/Misc/AutomationTest.cpp
+++ b/Engine/Source/Runtime/Core/Private/Misc/AutomationTest.cpp
@@ -2669,6 +2669,46 @@ bool FAutomationTestBase::TestGreaterEqual(const TCHAR* What, const double Actua
 	return false;
 }
 
+bool FAutomationTestBase::TestBetweenExclusive(const TCHAR* What, const int32 Actual, const int32 ExpectedLow, const int32 ExpectedHigh)
+{
+	if (Actual > ExpectedLow && Actual < ExpectedHigh)
+	{
+		return true;
+	}
+	AddError(FString::Printf(TEXT("Expected '%s' to be between %d and %d exclusive, but it was %d."), What, ExpectedLow, ExpectedHigh, Actual), 1);
+	return false;
+}
+
+bool FAutomationTestBase::TestBetweenExclusive(const TCHAR* What, const int64 Actual, const int64 ExpectedLow, const int64 ExpectedHigh)
+{
+	if (Actual > ExpectedLow && Actual < ExpectedHigh)
+	{
+		return true;
+	}
+	AddError(FString::Printf(TEXT("Expected '%s' to be between %" PRId64 " and %" PRId64 " exclusive, but it was %" PRId64 "."), What, ExpectedLow, ExpectedHigh, Actual), 1);
+	return false;
+}
+
+bool FAutomationTestBase::TestBetweenInclusive(const TCHAR* What, const int32 Actual, const int32 ExpectedLow, const int32 ExpectedHigh)
+{
+	if (Actual >= ExpectedLow && Actual <= ExpectedHigh)
+	{
+		return true;
+	}
+	AddError(FString::Printf(TEXT("Expected '%s' to be between %d and %d inclusive, but it was %d."), What, ExpectedLow, ExpectedHigh, Actual), 1);
+	return false;
+}
+
+bool FAutomationTestBase::TestBetweenInclusive(const TCHAR* What, const int64 Actual, const int64 ExpectedLow, const int64 ExpectedHigh)
+{
+	if (Actual >= ExpectedLow && Actual <= ExpectedHigh)
+	{
+		return true;
+	}
+	AddError(FString::Printf(TEXT("Expected '%s' to be between %" PRId64 " and %" PRId64 " inclusive, but it was %" PRId64 "."), What, ExpectedLow, ExpectedHigh, Actual), 1);
+	return false;
+}
+
 bool FAutomationTestBase::TestFalse(const TCHAR* What, bool Value)
 {
 	if (Value)
diff --git a/Engine/Source/Runtime/Core/Public/Misc/AutomationTest.h b/Engine/Source/Runtime/Core/Public/Misc/AutomationTest.h
--- a/Engine/Source/Runtime/Core/Public/Misc/AutomationTest.h
+++ b/Engine/Source/Runtime/Core/Public/Misc/AutomationTest.h
@@ -2181,6 +2181,31 @@ public:
 		return TestNotEqualSensitive(*What, Tmp, Expected);
 	}
 
+	struct FHasToString
+	{
+		template<typename ValueType>
+		auto Requires(ValueType& InElem) -> decltype(InElem.ToString());
+	};
+
+	template<typename ValueType> 
+	inline void ReportError(const TCHAR* What, const ValueType& Actual, const ValueType& Expected)
+	{
+		if constexpr (TModels_V<FHasToString, ValueType>)
+		{
+			ReportError(What, *Actual.ToString(), *Expected.ToString());
+		}
+		else
+		{
+			AddError(FString::Printf(TEXT("%s: The two values are not equal."), What));
+		}
+	}
+
+	template<> 
+	inline void ReportError(const TCHAR* What, const FString& Actual, const FString& Expected)
+	{
+		AddError(FString::Printf(TEXT("%s: The two values are not equal (actual=%s, expected=%s)."), What, *Actual, *Expected));
+	}
+
 	/**
 	 * Logs an error if the two values are not equal.
 	 *
@@ -2195,7 +2220,7 @@ public:
 	{
 		if (Actual != Expected)
 		{
-			AddError(FString::Printf(TEXT("%s: The two values are not equal."), What));
+			ReportError(What, Actual, Expected);
 			return false;
 		}
 		return true;
@@ -2339,6 +2364,31 @@ public:
 		return TestGreaterEqual(*What, Actual, Expected, Tolerance);
 	}
 
+	CORE_API bool TestBetweenExclusive(const TCHAR* What, const int32 Actual, const int32 ExpectedLow, const int32 ExpectedHigh);
+	CORE_API bool TestBetweenExclusive(const TCHAR* What, const int64 Actual, const int64 ExpectedLow, const int64 ExpectedHigh);
+	CORE_API bool TestBetweenInclusive(const TCHAR* What, const int32 Actual, const int32 ExpectedLow, const int32 ExpectedHigh);
+	CORE_API bool TestBetweenInclusive(const TCHAR* What, const int64 Actual, const int64 ExpectedLow, const int64 ExpectedHigh);
+
+	bool TestBetweenExclusive(const FString& What, const int32 Actual, const int32 ExpectedLow, const int32 ExpectedHigh)
+	{
+		return TestBetweenExclusive(*What, Actual, ExpectedLow, ExpectedHigh);
+	}
+
+	bool TestBetweenExclusive(const FString& What, const int64 Actual, const int64 ExpectedLow, const int64 ExpectedHigh)
+	{
+		return TestBetweenExclusive(*What, Actual, ExpectedLow, ExpectedHigh);
+	}
+
+	bool TestBetweenInclusive(const FString& What, const int32 Actual, const int32 ExpectedLow, const int32 ExpectedHigh)
+	{
+		return TestBetweenInclusive(*What, Actual, ExpectedLow, ExpectedHigh);
+	}
+
+	bool TestBetweenInclusive(const FString& What, const int64 Actual, const int64 ExpectedLow, const int64 ExpectedHigh)
+	{
+		return TestBetweenInclusive(*What, Actual, ExpectedLow, ExpectedHigh);
+	}
+
 	/**
 	 * Logs an error if the specified Boolean value is not false.
 	 *
