[Android] Platform file memory-mapping robustness

Made mapped regions debug counter atomic as those could be mapped concurrently.

Also made the outstanding regions check in mapped handle destructor fatal in non-shipping builds to catch issues earlier. This matches implementation with other platforms.

#rnx
#Synced-CL 51402227

diff --git a/Engine/Source/Runtime/Core/Private/Android/AndroidPlatformFile.cpp b/Engine/Source/Runtime/Core/Private/Android/AndroidPlatformFile.cpp
--- a/Engine/Source/Runtime/Core/Private/Android/AndroidPlatformFile.cpp
+++ b/Engine/Source/Runtime/Core/Private/Android/AndroidPlatformFile.cpp
@@ -28,6 +28,8 @@
 #include "Misc/ScopeExit.h"
 #include <Async/MappedFileHandle.h>
 
+#include <atomic>
+
 #if PLATFORM_USE_PLATFORM_FILE_MANAGED_STORAGE_WRAPPER
 #include "HAL/IPlatformFileManagedStorageWrapper.h"
 #endif //PLATFORM_USE_PLATFORM_FILE_MANAGED_STORAGE_WRAPPER
@@ -1359,7 +1361,17 @@ public:
 
 	virtual ~FAndroidMappedFileHandle() override
 	{
-		check(!NumOutstandingRegions); // can't delete the file before you delete all outstanding regions
+		// can't delete the file before you delete all outstanding regions
+		UE_CLOGF(
+				NumOutstandingRegions.load(std::memory_order_relaxed),
+				LogAndroidFile,
+#if UE_BUILD_SHIPPING
+				Error,
+#else
+				Fatal,
+#endif
+				"Cleaning mapped file with alive mapped regions: %ls",
+				*Filename);
 	}
 
 	virtual IMappedFileRegion* MapRegion(int64 Offset = 0, int64 BytesToMap = MAX_int64, FFileMappingFlags Flags = EMappedFileFlags::ENone) override
@@ -1398,14 +1410,14 @@ public:
 		// create a mapping for this range
 		const uint8* MapPtr = AlignedMapPtr + Offset - AlignedOffset;
 		FAndroidMappedFileRegion* Result = new FAndroidMappedFileRegion(MapPtr, AlignedMapPtr, BytesToMap, AlignedSize, Filename, Offset, this);
-		NumOutstandingRegions++;
+		NumOutstandingRegions.fetch_add(1, std::memory_order_relaxed);
 		return Result;
 	}
 
 	void UnMap(const FAndroidMappedFileRegion* Region)
 	{
-		check(NumOutstandingRegions > 0);
-		NumOutstandingRegions--;
+		int32 OldNumOutstandingRegions = NumOutstandingRegions.fetch_sub(1, std::memory_order_relaxed);
+		check(OldNumOutstandingRegions > 0);
 
 		const int Res = munmap(const_cast<uint8*>(Region->AlignedPtr), Region->AlignedSize);
 #if LOG_ANDROID_FILE
@@ -1418,7 +1430,7 @@ public:
 private:
 	const uint8* MappedPtr;
 	FString Filename;
-	int32 NumOutstandingRegions;
+	std::atomic<int32> NumOutstandingRegions;
 	FAndroidUniqueFd FileHandle;
 
 	int64 GetCurrentFileSize() const
