Commonize IOS network monitoring to all Apple platforms. This resolves the fact that network status / network type platform queries were not implemented for MacOS.

#Synced-CL 51099106

diff --git a/Engine/Source/Runtime/Core/Private/Apple/ApplePlatformMisc.cpp b/Engine/Source/Runtime/Core/Private/Apple/ApplePlatformMisc.cpp
--- a/Engine/Source/Runtime/Core/Private/Apple/ApplePlatformMisc.cpp
+++ b/Engine/Source/Runtime/Core/Private/Apple/ApplePlatformMisc.cpp
@@ -19,8 +19,11 @@
 #include "Apple/ApplePlatformCrashContext.h"
 #include "FramePro/FrameProProfiler.h"
 #include "CoreGlobals.h"
+#include "Async/TaskGraphInterfaces.h"
+#include "Misc/CoreDelegates.h"
 
 #include <os/log.h>
+#import <Network/Network.h>
 
 DEFINE_LOG_CATEGORY_STATIC(LogApplePlatformMisc, Log, All);
 
@@ -419,4 +422,184 @@ void FApplePlatformMisc::EndNamedEvent()
 #endif // FRAMEPRO_ENABLED
 }
 
+namespace UE::FApplePlatformMisc
+{
+	// Used to monitor network state. Replaces the deprecated SCNetworkReachability API.
+	// Note: defined here to not force objcisims into the header which is included on non-objc platforms
+	static nw_path_monitor_t PathMonitor = nullptr;
+	static nw_path_t CachedNetworkPath = nullptr;
+	static ENetworkConnectionType LastReportedNetworkType = ENetworkConnectionType::Unknown;
+	static ::FApplePlatformMisc::FNetworkConnectionCharacteristics CachedConnectionCharacteristics;
+
+	static TMulticastDelegate<void(::FApplePlatformMisc::FNetworkConnectionCharacteristics), FDefaultTSDelegateUserPolicy> OnNetworkConnectionCharacteristicsChanged;
+
+	// This is called from a whatever task graph thread picks.
+	static void NotifyAboutNetworkConnectionCharacteristicsChange(nw_path_t InPath)
+	{
+		::FApplePlatformMisc::FNetworkConnectionCharacteristics Details;
+		
+		Details.bSupportsDNS   = nw_path_has_dns(InPath);
+		Details.bSupportsIPv4  = nw_path_has_ipv4(InPath);
+		Details.bSupportsIPv6  = nw_path_has_ipv6(InPath);
+		Details.bIsConstrained = nw_path_is_constrained(InPath);
+		Details.bIsExpensive   = nw_path_is_expensive(InPath);
+		
+		// Path's reference count was increased to guarantee it lives
+		// til this moment of time. Refer to nw_path_monitor_set_update_handler
+		// callback below. However, we just used it, so:
+		nw_release(InPath);
+
+		CachedConnectionCharacteristics = Details;
+		OnNetworkConnectionCharacteristicsChanged.Broadcast(Details);
+	}
+
+	static ENetworkConnectionType CheckNetworkConnectionType()
+	{
+		nw_path_t MyPath = UE::FApplePlatformMisc::CachedNetworkPath;
+		if (MyPath == nullptr)
+		{
+			UE_LOG(LogInit, Warning, TEXT("Network path currently NOT set"));
+			return ENetworkConnectionType::Unknown;
+		}
+		
+		nw_path_status_t PathStatus = nw_path_get_status(MyPath);
+		if (PathStatus != nw_path_status_satisfied)
+		{
+			if (PathStatus == nw_path_status_invalid)
+			{
+				UE_LOG(LogInit, Warning, TEXT("Network path is invalid"));
+				return ENetworkConnectionType::None;
+			}
+			
+			// Basically, if there's _no_ network interfaces available, assume it's in airplane mode.
+			// Note: if you have wifi enabled - but airplane mode also enabled, that will NOT be detected
+			//       as airplane mode as you still have a valid network path.
+			__block bool HasAnInterface = false;
+			nw_path_enumerate_interfaces(MyPath, (nw_path_enumerate_interfaces_block_t) ^ (nw_interface_t Interface)
+			{
+				const char* Name = nw_interface_get_name(Interface);
+				// UE_LOG(LogIOS, Log, TEXT("NW path monitor available interface: %s"), *FString(Name));
+				HasAnInterface = true;
+				return false;
+			});
+			
+			if (!HasAnInterface)
+			{
+				return ENetworkConnectionType::AirplaneMode;
+			}
+			return ENetworkConnectionType::None;
+		}
+		else
+		{
+			// If satisfied detect the type of connection or maybe use isExpensive???
+			bool bHasActiveWiFiConnection = nw_path_uses_interface_type(MyPath, nw_interface_type_wifi);
+			bool bHasActiveCellConnection = nw_path_uses_interface_type(MyPath, nw_interface_type_cellular);
+			bool bHasActiveWiredConnection = nw_path_uses_interface_type(MyPath, nw_interface_type_wired);
+
+			if (bHasActiveWiFiConnection)
+			{
+				return ENetworkConnectionType::WiFi;
+			}
+			else if (bHasActiveCellConnection)
+			{
+				return ENetworkConnectionType::Cell;
+			}
+			else if (bHasActiveWiredConnection)
+			{
+				return ENetworkConnectionType::Ethernet;
+			}
+			else
+			{
+				// For now, mw_interface_type_loopback and nw_interface_type_other return "Unknown"
+				return ENetworkConnectionType::Unknown;
+			}
+		}
+	}
+} // UE::FApplePlatformMisc
+
+ENetworkConnectionType FApplePlatformMisc::GetNetworkConnectionType()
+{
+	static TOptional<ENetworkConnectionType> ConnectionType = {};
+	static double LastCheckTime = 0;
+
+	const double CurrentTime = FPlatformTime::Seconds();
+	const double CheckInterval = 0.2;
+
+	if (!ConnectionType.IsSet() || CurrentTime >= LastCheckTime + CheckInterval)
+	{
+		ConnectionType = UE::FApplePlatformMisc::CheckNetworkConnectionType();
+		LastCheckTime = CurrentTime;
+	}
+
+	ensure(ConnectionType.IsSet());
+	return ConnectionType.GetValue();
+}
+
+bool FApplePlatformMisc::IsDataSavingNetworkConnection()
+{
+	return GetNetworkConnectionCharacteristics().bIsConstrained;
+}
+
+FApplePlatformMisc::FNetworkConnectionCharacteristics FApplePlatformMisc::GetNetworkConnectionCharacteristics()
+{
+	return UE::FApplePlatformMisc::CachedConnectionCharacteristics;
+}
+
+bool FApplePlatformMisc::HasActiveWiFiConnection()
+{
+	return GetNetworkConnectionType() == ENetworkConnectionType::WiFi;
+}
+
+TMulticastDelegateRegistration<void(FApplePlatformMisc::FNetworkConnectionCharacteristics), FDefaultTSDelegateUserPolicy>& FApplePlatformMisc::OnNetworkConnectionCharacteristicsChanged()
+{
+	return UE::FApplePlatformMisc::OnNetworkConnectionCharacteristicsChanged;
+}
+
+void FApplePlatformMisc::PlatformInit()
+{
+	// start monitoring network status
+	UE::FApplePlatformMisc::PathMonitor = nw_path_monitor_create();
+	nw_path_monitor_set_update_handler(
+		UE::FApplePlatformMisc::PathMonitor, 
+		^(nw_path_t InPath)
+		{
+			bool bIsConstrained = nw_path_is_constrained(InPath);
+			bool bIsExpensive   = nw_path_is_expensive(InPath);
+
+			UE_LOG(LogInit, Log, TEXT("Network path changed (IsConstrained=%d, IsExpensive=%d)"), bIsConstrained, bIsExpensive);
+			
+			// cache the path we get
+			if (UE::FApplePlatformMisc::CachedNetworkPath != nullptr)
+			{
+				nw_release(UE::FApplePlatformMisc::CachedNetworkPath);
+			}
+			UE::FApplePlatformMisc::CachedNetworkPath = InPath;
+			nw_retain(UE::FApplePlatformMisc::CachedNetworkPath);
+			if (FTaskGraphInterface::IsRunning())
+			{
+				ENetworkConnectionType ConnectionType = UE::FApplePlatformMisc::CheckNetworkConnectionType();
+				if (ConnectionType != UE::FApplePlatformMisc::LastReportedNetworkType)
+				{
+					UE::FApplePlatformMisc::LastReportedNetworkType = ConnectionType;
+					FFunctionGraphTask::CreateAndDispatchWhenReady([ConnectionType]() {
+						FCoreDelegates::OnNetworkConnectionChanged.Broadcast(ConnectionType);
+					}, TStatId(), nullptr, ENamedThreads::GameThread);
+				}
+				
+				// Adding a reference, so path survives to the task body.
+				nw_retain(InPath);
+				FFunctionGraphTask::CreateAndDispatchWhenReady(
+				// nw_path_t is actually a pointer, so captured by value.
+				[InPath] ()
+				{
+					UE::FApplePlatformMisc::NotifyAboutNetworkConnectionCharacteristicsChange(InPath);
+				}, TStatId());
+			}
+		}
+	);
+	
+	nw_path_monitor_set_queue(UE::FApplePlatformMisc::PathMonitor, dispatch_get_main_queue());
+	nw_path_monitor_start(UE::FApplePlatformMisc::PathMonitor);
+
+}
 #endif // APPLE_PROFILING_ENABLED
diff --git a/Engine/Source/Runtime/Core/Private/IOS/IOSPlatformMisc.cpp b/Engine/Source/Runtime/Core/Private/IOS/IOSPlatformMisc.cpp
--- a/Engine/Source/Runtime/Core/Private/IOS/IOSPlatformMisc.cpp
+++ b/Engine/Source/Runtime/Core/Private/IOS/IOSPlatformMisc.cpp
@@ -8,7 +8,6 @@
 
 #include "Apple/AppleStringUtils.h"
 #include "Apple/ApplePlatformMisc.h"
-#include "Async/TaskGraphInterfaces.h"
 #include "Apple/ApplePlatformCrashContext.h"
 #include "Delegates/Delegate.h"
 #include "HAL/ExceptionHandling.h"
@@ -55,7 +54,6 @@
 #endif
 
 #import <UserNotifications/UserNotifications.h>
-#import <Network/Network.h>
 
 #include <sys/sysctl.h> // sysctlbyname
 
@@ -99,101 +97,6 @@ static int32 GetFreeMemoryMB()
 	return MemoryStats.AvailablePhysical / 1024 / 1024;
 }
 
-namespace UE::FIOSPlatformMisc
-{
-	// Used to monitor network state. Replaces the deprecated SCNetworkReachability API.
-	// Note: defined here to not force objcisims into the header which is included on non-objc platforms
-	static nw_path_monitor_t PathMonitor = nullptr;
-	static nw_path_t CachedNetworkPath = nullptr;
-	static ENetworkConnectionType LastReportedNetworkType = ENetworkConnectionType::Unknown;
-	static ::FIOSPlatformMisc::FNetworkConnectionCharacteristics CachedConnectionCharacteristics;
-
-	static TMulticastDelegate<void(::FIOSPlatformMisc::FNetworkConnectionCharacteristics), FDefaultTSDelegateUserPolicy> OnNetworkConnectionCharacteristicsChanged;
-
-	// This is called from a whatever task graph thread picks.
-	static void NotifyAboutNetworkConnectionCharacteristicsChange(nw_path_t InPath)
-	{
-		::FIOSPlatformMisc::FNetworkConnectionCharacteristics Details;
-		
-		Details.bSupportsDNS   = nw_path_has_dns(InPath);
-		Details.bSupportsIPv4  = nw_path_has_ipv4(InPath);
-		Details.bSupportsIPv6  = nw_path_has_ipv6(InPath);
-		Details.bIsConstrained = nw_path_is_constrained(InPath);
-		Details.bIsExpensive   = nw_path_is_expensive(InPath);
-		
-		// Path's reference count was increased to guarantee it lives
-		// til this moment of time. Refer to nw_path_monitor_set_update_handler
-		// callback below. However, we just used it, so:
-		nw_release(InPath);
-
-		CachedConnectionCharacteristics = Details;
-		OnNetworkConnectionCharacteristicsChanged.Broadcast(Details);
-	}
-
-	static ENetworkConnectionType CheckNetworkConnectionType()
-	{
-		nw_path_t MyPath = UE::FIOSPlatformMisc::CachedNetworkPath;
-		if (MyPath == nullptr)
-		{
-			UE_LOG(LogInit, Warning, TEXT("Network path currently NOT set"));
-			return ENetworkConnectionType::Unknown;
-		}
-		
-		nw_path_status_t PathStatus = nw_path_get_status(MyPath);
-		if (PathStatus != nw_path_status_satisfied)
-		{
-			if (PathStatus == nw_path_status_invalid)
-			{
-				UE_LOG(LogInit, Warning, TEXT("Network path is invalid"));
-				return ENetworkConnectionType::None;
-			}
-			
-			// Basically, if there's _no_ network interfaces available, assume it's in airplane mode.
-			// Note: if you have wifi enabled - but airplane mode also enabled, that will NOT be detected
-			//       as airplane mode as you still have a valid network path.
-			__block bool HasAnInterface = false;
-			nw_path_enumerate_interfaces(MyPath, (nw_path_enumerate_interfaces_block_t) ^ (nw_interface_t Interface)
-			{
-				const char* Name = nw_interface_get_name(Interface);
-				// UE_LOG(LogIOS, Log, TEXT("NW path monitor available interface: %s"), *FString(Name));
-				HasAnInterface = true;
-				return false;
-			});
-			
-			if (!HasAnInterface)
-			{
-				return ENetworkConnectionType::AirplaneMode;
-			}
-			return ENetworkConnectionType::None;
-		}
-		else
-		{
-			// If satisfied detect the type of connection or maybe use isExpensive???
-			bool bHasActiveWiFiConnection = nw_path_uses_interface_type(MyPath, nw_interface_type_wifi);
-			bool bHasActiveCellConnection = nw_path_uses_interface_type(MyPath, nw_interface_type_cellular);
-			bool bHasActiveWiredConnection = nw_path_uses_interface_type(MyPath, nw_interface_type_wired);
-
-			if (bHasActiveWiFiConnection)
-			{
-				return ENetworkConnectionType::WiFi;
-			}
-			else if (bHasActiveCellConnection)
-			{
-				return ENetworkConnectionType::Cell;
-			}
-			else if (bHasActiveWiredConnection)
-			{
-				return ENetworkConnectionType::Ethernet;
-			}
-			else
-			{
-				// For now, mw_interface_type_loopback and nw_interface_type_other return "Unknown"
-				return ENetworkConnectionType::Unknown;
-			}
-		}
-	}
-} // UE::FIOSPlatformMisc
-
 void FIOSPlatformMisc::PlatformInit()
 {
 	// PlatformInit() starts the UI thread which creates the framebuffer and it requires
@@ -288,55 +191,9 @@ void FIOSPlatformMisc::PlatformInit()
 		NSLog(@"Error excluding %@ from backup %@",[URL lastPathComponent], error);
 	}
 	
-	// start monitoring network status
-	UE::FIOSPlatformMisc::PathMonitor = nw_path_monitor_create();
-	nw_path_monitor_set_update_handler(
-		UE::FIOSPlatformMisc::PathMonitor, 
-		^(nw_path_t InPath)
-		{
-			bool bIsConstrained = nw_path_is_constrained(InPath);
-			bool bIsExpensive   = nw_path_is_expensive(InPath);
-
-			UE_LOG(LogInit, Log, TEXT("Network path changed (IsConstrained=%d, IsExpensive=%d)"), bIsConstrained, bIsExpensive);
-			
-			// cache the path we get
-			if (UE::FIOSPlatformMisc::PathMonitor != nullptr)
-			{
-				nw_release(UE::FIOSPlatformMisc::CachedNetworkPath);
-			}
-			UE::FIOSPlatformMisc::CachedNetworkPath = InPath;
-			nw_retain(UE::FIOSPlatformMisc::CachedNetworkPath);
-			if (FTaskGraphInterface::IsRunning())
-			{
-				ENetworkConnectionType ConnectionType = UE::FIOSPlatformMisc::CheckNetworkConnectionType();
-				if (ConnectionType != UE::FIOSPlatformMisc::LastReportedNetworkType)
-				{
-					UE::FIOSPlatformMisc::LastReportedNetworkType = ConnectionType;
-					FFunctionGraphTask::CreateAndDispatchWhenReady([ConnectionType]() {
-						FCoreDelegates::OnNetworkConnectionChanged.Broadcast(ConnectionType);
-					}, TStatId(), nullptr, ENamedThreads::GameThread);
-				}
-				
-				// Adding a reference, so path survives to the task body.
-				nw_retain(InPath);
-				FFunctionGraphTask::CreateAndDispatchWhenReady(
-				// nw_path_t is actually a pointer, so captured by value.
-				[InPath] ()
-				{
-					UE::FIOSPlatformMisc::NotifyAboutNetworkConnectionCharacteristicsChange(InPath);
-				}, TStatId());
-			}
-		}
-	);
-	
-	nw_path_monitor_set_queue(UE::FIOSPlatformMisc::PathMonitor, dispatch_get_main_queue());
-	nw_path_monitor_start(UE::FIOSPlatformMisc::PathMonitor);
+	FApplePlatformMisc::PlatformInit();
 }
 
-TMulticastDelegateRegistration<void(FIOSPlatformMisc::FNetworkConnectionCharacteristics), FDefaultTSDelegateUserPolicy>& FIOSPlatformMisc::OnNetworkConnectionCharacteristicsChanged()
-{
-	return UE::FIOSPlatformMisc::OnNetworkConnectionCharacteristicsChanged;
-}
 
 // Defines the PlatformFeatures module name for iOS, used by PlatformFeatures.h.
 const TCHAR* FIOSPlatformMisc::GetPlatformFeaturesModuleName()
@@ -1947,39 +1804,6 @@ void FIOSPlatformMisc::GetValidTargetPlatforms(TArray<FString>& TargetPlatformNa
 #endif
 }
 
-ENetworkConnectionType FIOSPlatformMisc::GetNetworkConnectionType()
-{
-	static TOptional<ENetworkConnectionType> ConnectionType = {};
-	static double LastCheckTime = 0;
-
-	const double CurrentTime = FPlatformTime::Seconds();
-	const double CheckInterval = 0.2;
-
-	if (!ConnectionType.IsSet() || CurrentTime >= LastCheckTime + CheckInterval)
-	{
-		ConnectionType = UE::FIOSPlatformMisc::CheckNetworkConnectionType();
-		LastCheckTime = CurrentTime;
-	}
-
-	ensure(ConnectionType.IsSet());
-	return ConnectionType.GetValue();
-}
-
-bool FIOSPlatformMisc::IsDataSavingNetworkConnection()
-{
-	return GetNetworkConnectionCharacteristics().bIsConstrained;
-}
-
-FIOSPlatformMisc::FNetworkConnectionCharacteristics FIOSPlatformMisc::GetNetworkConnectionCharacteristics()
-{
-	return UE::FIOSPlatformMisc::CachedConnectionCharacteristics;
-}
-
-bool FIOSPlatformMisc::HasActiveWiFiConnection()
-{
-    return GetNetworkConnectionType() == ENetworkConnectionType::WiFi;
-}
-
 FString FIOSPlatformMisc::GetCPUVendor()
 {
 	return TEXT("Apple");
diff --git a/Engine/Source/Runtime/Core/Private/Mac/MacPlatformMisc.cpp b/Engine/Source/Runtime/Core/Private/Mac/MacPlatformMisc.cpp
--- a/Engine/Source/Runtime/Core/Private/Mac/MacPlatformMisc.cpp
+++ b/Engine/Source/Runtime/Core/Private/Mac/MacPlatformMisc.cpp
@@ -623,6 +623,8 @@ void FMacPlatformMisc::PlatformInit()
 		UE_LOG(LogInit, Log, TEXT("No Xcode installed"));
 	}
 #endif
+	
+	FApplePlatformMisc::PlatformInit();
 }
 
 void FMacPlatformMisc::PostInitMacAppInfoUpdate()
diff --git a/Engine/Source/Runtime/Core/Public/Apple/ApplePlatformMisc.h b/Engine/Source/Runtime/Core/Public/Apple/ApplePlatformMisc.h
--- a/Engine/Source/Runtime/Core/Public/Apple/ApplePlatformMisc.h
+++ b/Engine/Source/Runtime/Core/Public/Apple/ApplePlatformMisc.h
@@ -31,6 +31,11 @@
 
 #define UE_DEBUG_BREAK_IMPL() PLATFORM_BREAK()
 
+template <typename FuncType, typename UserPolicy>
+class TMulticastDelegateRegistration;
+
+struct FDefaultTSDelegateUserPolicy;
+
 /**
 * Apple implementation of the misc OS functions
 **/
@@ -83,4 +88,28 @@ struct CORE_API FApplePlatformMisc : public FGenericPlatformMisc
 	static void ReleaseAutoreleasePool(void *Pool);
 	
 	static CORE_API void SetEnvironmentVar(const TCHAR* InVariableName, const TCHAR* Value);
+	
+	static ENetworkConnectionType GetNetworkConnectionType();
+	static bool IsDataSavingNetworkConnection();
+	static bool HasActiveWiFiConnection();
+	
+	// Information about the properties of the network that a connection uses.
+	struct FNetworkConnectionCharacteristics
+	{
+		// Indndicating whether the used network interface has a DNS server configured.
+		bool bSupportsDNS   : 1 = false;
+		// Indicating whether the used network interface can route IPv4 traffic.
+		bool bSupportsIPv4  : 1 = false;
+		// Indicating whether the used network interface can route IPv6 traffic.
+		bool bSupportsIPv6  : 1 = false;
+		// Indicating whether the used network interface is in Low Data Mode.
+		// It is applicable to both WiFi & cellular connections.
+		bool bIsConstrained : 1 = false;
+		// Indicating whether the used network interface is considered
+		// expensive, such as cellular or a personal hotspot.
+		bool bIsExpensive   : 1 = false;
+	};
+	// Thread-safe event that is fired from a task-thread.
+	static TMulticastDelegateRegistration<void(FNetworkConnectionCharacteristics), FDefaultTSDelegateUserPolicy>& OnNetworkConnectionCharacteristicsChanged();
+	static FNetworkConnectionCharacteristics GetNetworkConnectionCharacteristics();
 };
diff --git a/Engine/Source/Runtime/Core/Public/IOS/IOSPlatformMisc.h b/Engine/Source/Runtime/Core/Public/IOS/IOSPlatformMisc.h
--- a/Engine/Source/Runtime/Core/Public/IOS/IOSPlatformMisc.h
+++ b/Engine/Source/Runtime/Core/Public/IOS/IOSPlatformMisc.h
@@ -19,11 +19,6 @@
 template <typename FuncType>
 class TFunction;
 
-struct FDefaultTSDelegateUserPolicy;
-
-template <typename FuncType, typename UserPolicy>
-class TMulticastDelegateRegistration;
-
 /**
 * iOS implementation of the misc OS functions
 **/
@@ -50,9 +45,6 @@ struct CORE_API FIOSPlatformMisc : public FApplePlatformMisc
 	static bool DeleteStoredValue(const FString& InStoreId, const FString& InSectionName, const FString& InKeyName);
 	static bool DeleteStoredSection(const FString& InStoreId, const FString& InSectionName);
 	static void GetValidTargetPlatforms(class TArray<class FString>& TargetPlatformNames);
-	static ENetworkConnectionType GetNetworkConnectionType();
-	static bool IsDataSavingNetworkConnection();
-	static bool HasActiveWiFiConnection();
 	static const TCHAR* GamePersistentDownloadDir();
     static bool HasSeparateChannelForDebugOutput();
 
@@ -246,26 +238,7 @@ struct CORE_API FIOSPlatformMisc : public FApplePlatformMisc
 
 	static bool RequestDeviceCheckToken(TFunction<void(const TArray<uint8>&)> QuerySucceededFunc, TFunction<void(const FString&, const FString&)> QueryFailedFunc);
 	
-	// Information about the properties of the network that a connection uses.
-	struct FNetworkConnectionCharacteristics
-	{
-		// Indndicating whether the used network interface has a DNS server configured.
-		bool bSupportsDNS   : 1 = false;
-		// Indicating whether the used network interface can route IPv4 traffic.
-		bool bSupportsIPv4  : 1 = false;
-		// Indicating whether the used network interface can route IPv6 traffic.
-		bool bSupportsIPv6  : 1 = false;
-		// Indicating whether the used network interface is in Low Data Mode.
-		// It is applicable to both WiFi & cellular connections.
-		bool bIsConstrained : 1 = false;
-		// Indicating whether the used network interface is considered
-		// expensive, such as cellular or a personal hotspot.
-		bool bIsExpensive   : 1 = false;
-	};
-	// Thread-safe event that is fired from a task-thread.
-	static TMulticastDelegateRegistration<void(FNetworkConnectionCharacteristics), FDefaultTSDelegateUserPolicy>& OnNetworkConnectionCharacteristicsChanged();
-	static FNetworkConnectionCharacteristics GetNetworkConnectionCharacteristics();
-    
+   
 	FORCEINLINE static void ChooseHDRDeviceAndColorGamut(uint32 DeviceId, uint32 DisplayNitLevel, EDisplayOutputFormat& OutputDevice, EDisplayColorGamut& ColorGamut)
 	{
 		// Linear output to Apple's specific format.
