Memory: Changed FCriticalSection to FPlatformRecursiveMutex in virtual memory code

This supports changing FCriticalSection to FRecursiveMutex, which depends on ParkingLot, which must not depend on itself for memory allocation.

#Synced-CL 50967057

diff --git a/Engine/Source/Runtime/Core/Private/Android/AndroidPlatformMemory.cpp b/Engine/Source/Runtime/Core/Private/Android/AndroidPlatformMemory.cpp
--- a/Engine/Source/Runtime/Core/Private/Android/AndroidPlatformMemory.cpp
+++ b/Engine/Source/Runtime/Core/Private/Android/AndroidPlatformMemory.cpp
@@ -2,13 +2,14 @@
 
 #include "Android/AndroidPlatformMemory.h"
 #include "Android/AndroidHeapProfiling.h"
+#include "Async/UniqueLock.h"
 #include "HAL/LowLevelMemTracker.h"
 #include "HAL/MallocBinned.h"
 #include "HAL/MallocBinned2.h"
 #include "HAL/MallocBinned3.h"
 #include "HAL/MallocMimalloc.h"
 #include "HAL/MallocAnsi.h"
-#include "Misc/ScopeLock.h"
+#include "HAL/PlatformMutex.h"
 #include "unistd.h"
 #include <jni.h>
 #include <sys/sysinfo.h>
@@ -508,8 +509,8 @@ void* FAndroidPlatformMemory::BinnedAllocFromOS(SIZE_T Size)
 	void* Ptr;
 	if (USE_MALLOC_BINNED2)
 	{
-		static FCriticalSection CriticalSection;
-		FScopeLock Lock(&CriticalSection);
+		static UE::FPlatformRecursiveMutex Mutex;
+		UE::TUniqueLock Lock(Mutex);
 
 		const FPlatformMemoryConstants& MemoryConstants = FPlatformMemory::GetConstants();
 
diff --git a/Engine/Source/Runtime/Core/Public/HAL/PageCache.h b/Engine/Source/Runtime/Core/Public/HAL/PageCache.h
--- a/Engine/Source/Runtime/Core/Public/HAL/PageCache.h
+++ b/Engine/Source/Runtime/Core/Public/HAL/PageCache.h
@@ -5,26 +5,26 @@
 // HEADER_UNIT_SKIP - Not included directly
 
 #include "CoreTypes.h"
-#include "HAL/CriticalSection.h"
+#include "Async/UniqueLock.h"
 #include "HAL/MallocBinnedCommon.h"
+#include "HAL/PlatformMutex.h"
 
 class FPageCache
 {
-	FCriticalSection CriticalSection;
+	UE::FPlatformRecursiveMutex Mutex;
 
 	uint8* LowAddress;
 	uint8* HighAddress;
 	size_t PageSize;
-	uint32 NumPages;
-	uint32 MemSize;
-	uint32 SweepPage;
-	uint32 CommittedPages;
-	uint32 DecommittedPages;
-	uint32 PendingDecommittedPages;
-
-	uint64 CommitHits;
-	uint64 CommitMisses;
+	uint32 NumPages = 0;
+	uint32 MemSize = 0;
+	uint32 SweepPage = 0;
+	uint32 CommittedPages = 0;
+	uint32 DecommittedPages = 0;
+	uint32 PendingDecommittedPages = 0;
 
+	uint64 CommitHits = 0;
+	uint64 CommitMisses = 0;
 
 	FBitTree CurrentlyCommitted;
 	FBitTree NotPendingDecommit;
@@ -87,11 +87,6 @@ public:
 		: LowAddress((uint8*)InLowAdress)
 		, HighAddress((uint8*)InHighAddress)
 		, PageSize(InPageSize)
-		, NumPages(0)
-		, MemSize(0)
-		, SweepPage(0)
-		, CommittedPages(0)
-		, PendingDecommittedPages(0)
 	{
 	}
 
@@ -106,7 +101,7 @@ public:
 	void InitPageCache(void* Memory)
 	{
 		check(NumPages && MemSize);
-		FScopeLock Lock(&CriticalSection);
+		UE::TUniqueLock Lock(Mutex);
 		DecommittedPages = NumPages;
 		CurrentlyCommitted.FBitTreeInit(NumPages, Memory, MemSize, false);
 		NotPendingDecommit.FBitTreeInit(NumPages, (uint8*)Memory + MemSize, MemSize, true);
@@ -120,7 +115,7 @@ public:
 
 		uint32 NumFound = 0;
 
-		FScopeLock Lock(&CriticalSection);
+		UE::TUniqueLock Lock(Mutex);
 		// this loop could be accelerated by using the hierarchical info in the bit tree 
 		for (uint32 Index = StartPage; Index < EndPage; Index++)
 		{
@@ -143,7 +138,7 @@ public:
 
 		uint32 LastCommitPage = MAX_uint32;
 		uint32 StartCommitPage = MAX_uint32;
-		FScopeLock Lock(&CriticalSection);
+		UE::TUniqueLock Lock(Mutex);
 		// this loop could be accelerated by using the hierarchical info in the bit tree 
 		uint32 NumFound = 0;
 		uint32 NumUnPending = 0;
@@ -196,7 +191,7 @@ public:
 
 		uint32 LastCommitPage = MAX_uint32;
 		uint32 StartCommitPage = MAX_uint32;
-		FScopeLock Lock(&CriticalSection);
+		UE::TUniqueLock Lock(Mutex);
 		// this loop could be accelerated by using the hierarchical info in the bit tree 
 		uint32 NumFound = 0;
 		uint32 NumUnPending = 0;
@@ -247,7 +242,7 @@ public:
 		uint32 LastDecommitPage = MAX_uint32;
 		uint32 StartDecommitPage = MAX_uint32;
 		uint32 NumFound = 0;
-		FScopeLock Lock(&CriticalSection);
+		UE::TUniqueLock Lock(Mutex);
 
 		while (NumFound < NumNeed)
 		{
@@ -302,7 +297,7 @@ public:
 		uint32 LastDecommitPage = MAX_uint32;
 		uint32 StartDecommitPage = MAX_uint32;
 		uint32 NumFound = 0;
-		FScopeLock Lock(&CriticalSection);
+		UE::TUniqueLock Lock(Mutex);
 
 		auto DecommitRange = [&]()
 		{
@@ -360,12 +355,16 @@ public:
 	}
 	size_t GetFreeableMemory()
 	{
-		FScopeLock Lock(&CriticalSection);
+		UE::TUniqueLock Lock(Mutex);
 		return size_t(PendingDecommittedPages) * PageSize;
 	}
 	float GetHitRate()
 	{
-		FScopeLock Lock(&CriticalSection);
+		UE::TUniqueLock Lock(Mutex);
 		return 100.0f * float(CommitHits) / float(CommitHits + CommitMisses + 1); // +1 to avoid divide by zero
 	}
-};
\ No newline at end of file
+};
+
+#if UE_ENABLE_INCLUDE_ORDER_DEPRECATED_IN_5_8
+#	include "HAL/CriticalSection.h"
+#endif
diff --git a/Engine/Source/Runtime/Core/Public/HAL/VirtualAllocator.h b/Engine/Source/Runtime/Core/Public/HAL/VirtualAllocator.h
--- a/Engine/Source/Runtime/Core/Public/HAL/VirtualAllocator.h
+++ b/Engine/Source/Runtime/Core/Public/HAL/VirtualAllocator.h
@@ -3,10 +3,10 @@
 #pragma once
 
 #include "CoreTypes.h"
-#include "HAL/CriticalSection.h"
+#include "Async/UniqueLock.h"
 #include "HAL/ThreadSafeCounter.h"
+#include "HAL/PlatformMutex.h"
 #include "HAL/PlatformProcess.h"
-#include "Misc/ScopeLock.h"
 #include "Templates/AlignmentTemplates.h"
 
 class FVirtualAllocator
@@ -24,7 +24,7 @@ class FVirtualAllocator
 		FFreeLink* FirstFree = nullptr;
 	};
 
-	FCriticalSection CriticalSection;
+	UE::FPlatformRecursiveMutex Mutex;
 
 	uint8* LowAddress;
 	uint8* HighAddress;
@@ -68,6 +68,7 @@ class FVirtualAllocator
 				RecycledLinks = NewLink;
 			}
 			LinkSize += MaximumAlignment;
+			checkSlow(RecycledLinks); // To suppress a static analysis warning about RecycledLinks being null below.
 		}
 
 		FFreeLink* Link = RecycledLinks;
@@ -151,7 +152,7 @@ public:
 		}
 		FPerBlockSize& Block = Blocks[BlockIndex];
 
-		FScopeLock Lock(&CriticalSection);
+		UE::TUniqueLock Lock(Mutex);
 
 
 		uint8* Result;
@@ -210,7 +211,7 @@ public:
 		size_t AlignedSize = size_t(1) << BlockIndex;
 
 		FPerBlockSize& Block = Blocks[BlockIndex];
-		FScopeLock Lock(&CriticalSection);
+		UE::TUniqueLock Lock(Mutex);
 		FreeVirtualByBlock(Ptr, Block, AlignedSize);
 	}
 
@@ -234,7 +235,7 @@ public:
 
 	void GetStats(FVirtualAllocatorStats& OutStats)
 	{
-		FScopeLock Lock(&CriticalSection);
+		UE::TUniqueLock Lock(Mutex);
 		OutStats.PageSize = PageSize;
 		OutStats.MaximumAlignment = MaximumAlignment;
 		OutStats.VMSpaceTotal = HighAddress - LowAddress;
@@ -248,4 +249,9 @@ public:
 			OutStats.BlockStats[Index].FreeBlocksSize = Blocks[Index].FreeBlocksSize;
 		}
 	}
-};
\ No newline at end of file
+};
+
+#if UE_ENABLE_INCLUDE_ORDER_DEPRECATED_IN_5_8
+#	include "HAL/CriticalSection.h"
+#	include "Misc/ScopeLock.h"
+#endif
