[Core] Add guard to only allow arrayview function on compact sets be available for explicit uses of it.

#test Compiled array view against tset and tcompactset

#Synced-CL 51470662

diff --git a/Engine/Source/Runtime/Core/Public/Containers/CompactSet.h b/Engine/Source/Runtime/Core/Public/Containers/CompactSet.h
--- a/Engine/Source/Runtime/Core/Public/Containers/CompactSet.h
+++ b/Engine/Source/Runtime/Core/Public/Containers/CompactSet.h
@@ -2,6 +2,8 @@
 
 #pragma once
 
+#define UE_TCOMPACT_SET_FORCE_ARRAY_VIEW // always enabled for explicitly typed compact sets
 #define UE_TCOMPACT_SET TCompactSet
 #include "CompactSet.h.inl"
-#undef UE_TCOMPACT_SET
\ No newline at end of file
+#undef UE_TCOMPACT_SET
+#undef UE_TCOMPACT_SET_FORCE_ARRAY_VIEW
\ No newline at end of file
diff --git a/Engine/Source/Runtime/Core/Public/Containers/CompactSet.h.inl b/Engine/Source/Runtime/Core/Public/Containers/CompactSet.h.inl
--- a/Engine/Source/Runtime/Core/Public/Containers/CompactSet.h.inl
+++ b/Engine/Source/Runtime/Core/Public/Containers/CompactSet.h.inl
@@ -25,6 +25,18 @@ namespace UE::Core::Private
 	[[noreturn]] CORE_API void OnInvalidSetNum(unsigned long long NewNum);
 }
 
+// Array view is a unique feature for compact sets since the data is not allowed to have holes like sparse set
+// However to maintain compatibilty with sparse set this feature is not available by default
+#if !defined(UE_TCOMPACT_SET_ARRAY_VIEW)
+#define UE_TCOMPACT_SET_ARRAY_VIEW 0
+#endif
+
+#if UE_TCOMPACT_SET_ARRAY_VIEW || defined(UE_TCOMPACT_SET_FORCE_ARRAY_VIEW)
+#define UE_TCOMPACT_SET_ARRAY_VIEW_ENABLED 1
+#else
+#define UE_TCOMPACT_SET_ARRAY_VIEW_ENABLED 0
+#endif
+
 #define TSETPRIVATEFRIEND UE_JOIN(UE_TCOMPACT_SET, PrivateFriend)
 
 /**
@@ -1303,17 +1315,19 @@ public:
 		return bIncludesSet;
 	}
 
-	/** @return a TArray that's a copy of all the elements in this set, prefer using ArrayView where possible. */
+	/** @return a TArray that's a copy of all the elements in this set, prefer using ArrayView if available. */
 	[[nodiscard]] TArray<ElementType> Array() const
 	{
 		return TArray<ElementType>(GetData(), this->NumElements);
 	}
 
+#if UE_TCOMPACT_SET_ARRAY_VIEW_ENABLED
 	/** @return a readonly TArrayView of the elements */
 	[[nodiscard]] TArrayView<const ElementType> ArrayView() const
 	{
 		return TArrayView<const ElementType>(GetData(), this->NumElements);
 	}
+#endif
 
 	/** Creates an iterator for the contents of this set */
 	[[nodiscard]] UE_FORCEINLINE_HINT TIterator CreateIterator()
@@ -1784,4 +1798,5 @@ template <typename ElementType, typename KeyFuncs, typename Allocator> struct TI
 template <typename ElementType, typename KeyFuncs, typename Allocator> struct TIsTSet<      volatile UE_TCOMPACT_SET<ElementType, KeyFuncs, Allocator>> { enum { Value = true }; };
 template <typename ElementType, typename KeyFuncs, typename Allocator> struct TIsTSet<const volatile UE_TCOMPACT_SET<ElementType, KeyFuncs, Allocator>> { enum { Value = true }; };
 
-#undef TSETPRIVATEFRIEND
\ No newline at end of file
+#undef TSETPRIVATEFRIEND
+#undef UE_TCOMPACT_SET_ARRAY_VIEW_ENABLED
\ No newline at end of file
