Move LoadOnlyHostAndTargetPlatformModules to the DDPI area so we only return the platforms we will actually be operation on
When the command line is around, we only return Host + Target + any IniParents that are part of these Host/Target/NewParents

#Synced-CL 51322414

diff --git a/Engine/Source/Developer/TargetPlatform/Private/TargetPlatformManagerModule.cpp b/Engine/Source/Developer/TargetPlatform/Private/TargetPlatformManagerModule.cpp
--- a/Engine/Source/Developer/TargetPlatform/Private/TargetPlatformManagerModule.cpp
+++ b/Engine/Source/Developer/TargetPlatform/Private/TargetPlatformManagerModule.cpp
@@ -1179,8 +1179,6 @@ protected:
 		}
 #endif
 
-		TSet<FName> AlreadyInitialized;
-		
 #if WITH_EDITOR
 		// Get the platform we are previewing if GMaxRHIShaderPlatform is a preview SP
 		FName PlatformNamePreview = NAME_None;
@@ -1192,85 +1190,6 @@ protected:
 			}
 		}
 		
-		if (IsRunningCookCommandlet() && FParse::Param(FCommandLine::Get(), TEXT("-LoadOnlyHostAndTargetPlatformModules")))
-		{
-			// For every given target platform module, attempt to find the containing plugin (e.g. IOSClient -> IOS).
-			//
-			// For successful loading:
-			// * The containing platform must exist.
-			// * The containing platform must be enabled for use.
-			// * The loaded modules after init must contain the requested platform.
-			//
-			// Any failure will fall back to loading all discovered modules.
-				
-			bool bSuccess = false;
-			const TMap<FName, FDataDrivenPlatformInfo>& AllPlatformInfos = FDataDrivenPlatformInfoRegistry::GetAllPlatformInfos();
-			
-			FString TargetPlatformCmdLine;
-			if (FParse::Value(FCommandLine::Get(), TEXT("TargetPlatform="), TargetPlatformCmdLine) && TargetPlatformCmdLine != TEXT("All"))
-			{
-				bSuccess = true;
-				TArray<FString> TargetPlatformNames;
-				TargetPlatformCmdLine.ParseIntoArray(TargetPlatformNames, TEXT("+"), true);
-				
-				for (const FString& TargetPlatformName: TargetPlatformNames)
-				{
-					FString Stripped = TargetPlatformName;
-					Stripped.RemoveFromEnd("Client");
-					Stripped.RemoveFromEnd("Server");
-					Stripped.RemoveFromEnd("Editor");
-					// Cook flavours, for e.g. Android, are defined as Android_ASTC etc.
-					Stripped.LeftInline(Stripped.Contains(TEXT("_")) ? Stripped.Find(TEXT("_")) : Stripped.Len());	
-					
-					FName PlatformToLoad(Stripped);
-					
-					if (!AllPlatformInfos.Contains(PlatformToLoad))
-					{
-						UE_LOG(LogTargetPlatformManager, Warning, TEXT("Assumed parent platform (%s) not found for %s"), *PlatformToLoad.ToString(), *TargetPlatformName);
-						bSuccess = false;
-						break;
-					}
-					
-					const FDataDrivenPlatformInfo& Info = AllPlatformInfos[PlatformToLoad];
-					
-					if (!Info.bEnabledForUse)
-					{
-						UE_LOG(LogTargetPlatformManager, Warning, TEXT("Assumed parent platform (%s) was not enabled for use"), *PlatformToLoad.ToString());
-						bSuccess = false;
-						break;
-					}
-					
-					if (InitializeSinglePlatform(PlatformToLoad, Info.AutoSDKPath))
-					{
-						AlreadyInitialized.Add(PlatformToLoad);
-						if (!PlatformsByName.Contains(FName(TargetPlatformName)))
-						{
-							UE_LOG(LogTargetPlatformManager, Warning, TEXT("Target platform %s does not appear to be part of %s"), *TargetPlatformName, *PlatformToLoad.ToString());
-							bSuccess = false;
-							break;
-						}
-					}
-					else
-					{
-						UE_LOG(LogTargetPlatformManager, Warning, TEXT("Assumed parent platform %s failed to initialize"), *PlatformToLoad.ToString());
-						bSuccess = false;
-					}
-				}
-			}
-			
-			if (bSuccess)
-			{
-				// Host platform is always required so ensure it is initialized if it wasn't the requested.
-				const FName HostPlatformName = FPlatformProperties::IniPlatformName();
-				if (!AlreadyInitialized.Contains(HostPlatformName))
-				{
-					InitializeSinglePlatform(HostPlatformName, AllPlatformInfos[HostPlatformName].AutoSDKPath);
-				}
-				
-				return;
-			}
-		}
-		
 #endif
 
 		// find a set of valid target platform names (the platform DataDrivenPlatformInfo.ini file was found indicates support for the platform 
@@ -1302,7 +1221,7 @@ protected:
 			}
 #endif
 
-			if (Info.bEnabledForUse && !AlreadyInitialized.Contains(PlatformName))
+			if (Info.bEnabledForUse)
 			{
 				InitializeSinglePlatform(PlatformName, Info.AutoSDKPath);
 			}
diff --git a/Engine/Source/Runtime/Core/Private/Misc/DataDrivenPlatformInfoRegistry.cpp b/Engine/Source/Runtime/Core/Private/Misc/DataDrivenPlatformInfoRegistry.cpp
--- a/Engine/Source/Runtime/Core/Private/Misc/DataDrivenPlatformInfoRegistry.cpp
+++ b/Engine/Source/Runtime/Core/Private/Misc/DataDrivenPlatformInfoRegistry.cpp
@@ -25,41 +25,147 @@ namespace
 #endif
 }
 
-static const TArray<FString>& GetDataDrivenIniFilenames()
+static FString GetParentIniFromDDPI(const FString& DDPIPath)
 {
-	static bool bHasSearchedForFiles = false;
-	static TArray<FString> DataDrivenIniFilenames;
+	FString IniParent;
+	FString IniContents;
+	if (FFileHelper::LoadFileToString(IniContents, *DDPIPath))
+	{
+		FConfigFile IniFile;
+		IniFile.ProcessInputFileContents(IniContents, DDPIPath);
 
-	if (bHasSearchedForFiles == false)
+		IniFile.GetString(TEXT("DataDrivenPlatformInfo"), TEXT("IniParent"), IniParent);
+	}
+	return IniParent;
+}
+
+static const TArray<FString> GetAllDataDrivenIniFilenames()
+{
+	TArray<FString> DataDrivenIniFilenames;
+
+	// look for the special files in any config subdirectories
+	IFileManager::Get().FindFilesRecursive(DataDrivenIniFilenames, *FPaths::EngineConfigDir(), TEXT("DataDrivenPlatformInfo.ini"), true, false);
+
+	// manually look through the platform directories - we can't use GetExtensionDirs(), since that function uses the results of this function 
+	TArray<FString> PlatformDirs;
+	IFileManager::Get().FindFiles(PlatformDirs, *FPaths::Combine(FPaths::EngineDir(), TEXT("Platforms"), TEXT("*")), false, true);
+
+	for (const FString& PlatformDir : PlatformDirs)
 	{
-		bHasSearchedForFiles = true;
+		FString IniPath = FPaths::Combine(FPaths::EnginePlatformExtensionDir(*PlatformDir), TEXT("Config/DataDrivenPlatformInfo.ini"));
+		if (IFileManager::Get().FileExists(*IniPath))
+		{
+			DataDrivenIniFilenames.Add(IniPath);
+		}
+	}
+
+	// look for the special files in any project config subdirectories
+	TArray<FString> ProjectPlatformDirs;
+	IFileManager::Get().FindFiles(ProjectPlatformDirs, *FPaths::Combine(FPaths::ProjectConfigDir(), TEXT("*")), false, true);
+	for (const FString& PlatformDir : ProjectPlatformDirs)
+	{
+		FString IniPath = FPaths::Combine(FPaths::ProjectPlatformExtensionDir(*PlatformDir), TEXT("Config/DataDrivenPlatformInfo.ini"));
+		if (IFileManager::Get().FileExists(*IniPath))
+		{
+			DataDrivenIniFilenames.AddUnique(IniPath); // need AddUnique because if there's no project specified then ProjectConfigDir will be EngineConfigDir
+		}
+	}
+
+	return DataDrivenIniFilenames;
+}
+
+static const TArray<FString> GetHostAndTargetDataDrivenIniFilenames()
+{
+	// grab all we find on disk, and filter only the Host/Target, likely a better way to only search on disk for our host/target but start with easier for now
+	TArray<FString> DataDrivenIniFilenames = GetAllDataDrivenIniFilenames();
+
+	FString TargetPlatformCmdLine;
+	if (FParse::Value(FCommandLine::Get(), TEXT("TargetPlatform="), TargetPlatformCmdLine) && TargetPlatformCmdLine != TEXT("All"))
+	{
+		TSet<FString> HostAndTargetPlatforms;
+		HostAndTargetPlatforms.Add(FString::Printf(TEXT("/%s/"), ANSI_TO_TCHAR(FPlatformProperties::IniPlatformName()))); // always add Host
+
+		TArray<FString> TargetPlatformNames;
+		TargetPlatformCmdLine.ParseIntoArray(TargetPlatformNames, TEXT("+"), true);
+
+		for (const FString& TargetPlatformName : TargetPlatformNames)
+		{
+			FString Stripped = TargetPlatformName;
+			Stripped.RemoveFromEnd("Client");
+			Stripped.RemoveFromEnd("Server");
+			Stripped.RemoveFromEnd("Editor");
+			// Cook flavours, for e.g. Android, are defined as Android_ASTC etc.
+			int32 UnderscoreIdx = INDEX_NONE;
+			if (Stripped.FindChar(TEXT('_'), UnderscoreIdx))
+			{
+				Stripped.LeftInline(UnderscoreIdx);
+			}
 
-		// look for the special files in any config subdirectories
-		IFileManager::Get().FindFilesRecursive(DataDrivenIniFilenames, *FPaths::EngineConfigDir(), TEXT("DataDrivenPlatformInfo.ini"), true, false);
+			// add any TargetPlatforms
+			HostAndTargetPlatforms.Add(FString::Printf(TEXT("/%s/"), *Stripped));
+		}
 
-		// manually look through the platform directories - we can't use GetExtensionDirs(), since that function uses the results of this function 
-		TArray<FString> PlatformDirs;
-		IFileManager::Get().FindFiles(PlatformDirs, *FPaths::Combine(FPaths::EngineDir(), TEXT("Platforms"), TEXT("*")), false, true);
+		TSet<FString> AllPlatforms(HostAndTargetPlatforms);
+		TSet<FString> PlatformsToCheck = HostAndTargetPlatforms;
 
-		for (const FString& PlatformDir : PlatformDirs)
+		while (PlatformsToCheck.Num() > 0)
 		{
-			FString IniPath = FPaths::Combine(FPaths::EnginePlatformExtensionDir(*PlatformDir), TEXT("Config/DataDrivenPlatformInfo.ini"));
-			if (IFileManager::Get().FileExists(*IniPath))
+			TSet<FString> NewParentPlatforms;
+
+			for (const FString& DataDrivenIniFilename : DataDrivenIniFilenames)
 			{
-				DataDrivenIniFilenames.Add(IniPath);
+				for (const FString& Platform : PlatformsToCheck)
+				{
+					if (DataDrivenIniFilename.Contains(Platform))
+					{
+						FString ParentIniPlatform = GetParentIniFromDDPI(DataDrivenIniFilename);
+						if (!ParentIniPlatform.IsEmpty() && !AllPlatforms.Contains(ParentIniPlatform))
+						{
+							FString ParentIniWithSlash = FString(TEXT("/") + ParentIniPlatform + TEXT("/"));
+							AllPlatforms.Add(ParentIniWithSlash);
+							NewParentPlatforms.Add(ParentIniWithSlash);
+						}
+						break;
+					}
+				}
 			}
+
+			PlatformsToCheck = MoveTemp(NewParentPlatforms);
 		}
 
-		// look for the special files in any project config subdirectories
-		TArray<FString> ProjectPlatformDirs;
-		IFileManager::Get().FindFiles(ProjectPlatformDirs, *FPaths::Combine(FPaths::ProjectConfigDir(), TEXT("*")), false, true);
-		for (const FString& PlatformDir : ProjectPlatformDirs)
+		// remove any paths that are not part our Host or TargetPlatform list
+		DataDrivenIniFilenames.RemoveAll([&AllPlatforms](const FString& Filename)
 		{
-			FString IniPath = FPaths::Combine(FPaths::ProjectPlatformExtensionDir(*PlatformDir), TEXT("Config/DataDrivenPlatformInfo.ini"));
-			if (IFileManager::Get().FileExists(*IniPath))
+			for (const FString& Component : AllPlatforms)
 			{
-				DataDrivenIniFilenames.AddUnique(IniPath); // need AddUnique because if there's no project specified then ProjectConfigDir will be EngineConfigDir
+				if (Filename.Contains(Component))
+				{
+					return false;
+				}
 			}
+			return true;
+		});
+	}
+
+	return DataDrivenIniFilenames;
+}
+
+static const TArray<FString>& GetDataDrivenIniFilenames()
+{
+	static bool bHasSearchedForFiles = false;
+	static TArray<FString> DataDrivenIniFilenames;
+
+	if (bHasSearchedForFiles == false)
+	{
+		bHasSearchedForFiles = true;
+
+		if (FParse::Param(FCommandLine::Get(), TEXT("-LoadOnlyHostAndTargetPlatformModules")))
+		{
+			DataDrivenIniFilenames = GetHostAndTargetDataDrivenIniFilenames();
+		}
+		else
+		{
+			DataDrivenIniFilenames = GetAllDataDrivenIniFilenames();
 		}
 	}
 
