Added an option to toggle MB2\3 thread local caches on\off in runtime

#Synced-CL 51465518

diff --git a/Engine/Source/Runtime/Core/Private/HAL/MallocBinnedCommon.cpp b/Engine/Source/Runtime/Core/Private/HAL/MallocBinnedCommon.cpp
--- a/Engine/Source/Runtime/Core/Private/HAL/MallocBinnedCommon.cpp
+++ b/Engine/Source/Runtime/Core/Private/HAL/MallocBinnedCommon.cpp
@@ -550,14 +550,16 @@ static FAutoConsoleVariableRef GMallocBinnedFlushRegisteredThreadCachesOnOneThre
 	GMallocBinnedFlushRegisteredThreadCachesOnOneThread,
 	TEXT("Whether or not to attempt to flush registered thread caches on one thread (enabled by default)."));
 
-#if UE_MBC_ALLOW_RUNTIME_TWEAKING
-
+#if UE_MBC_SUPPORT_RUNTIME_TLC_TOGGLE
 	int32 GMallocBinnedPerThreadCaches = UE_DEFAULT_GMallocBinnedPerThreadCaches;
 	static FAutoConsoleVariableRef GMallocBinnedPerThreadCachesCVar(
 		TEXT("MallocBinned.PerThreadCaches"),
 		GMallocBinnedPerThreadCaches,
 		TEXT("Enables per-thread caches of small (<= 32768 byte) allocations from FMallocBinned2/3")
 	);
+#endif
+
+#if UE_MBC_ALLOW_RUNTIME_TWEAKING
 
 	int32 GMallocBinnedBundleSize = UE_DEFAULT_GMallocBinnedBundleSize;
 	static FAutoConsoleVariableRef GMallocBinnedBundleSizeCVar(
@@ -586,7 +588,6 @@ static FAutoConsoleVariableRef GMallocBinnedFlushRegisteredThreadCachesOnOneThre
 		GMallocBinnedMaxBundlesBeforeRecycle,
 		TEXT("Number of freed bundles in the global recycler before it returns them to the system, per-block size. Limited by UE_DEFAULT_GBinned3MaxBundlesBeforeRecycle (currently 4)")
 	);
-
 #endif	//~UE_MBC_ALLOW_RUNTIME_TWEAKING
 
 uint32 FMallocBinnedCommonBase::BinnedTlsSlot = FPlatformTLS::InvalidTlsSlot;
diff --git a/Engine/Source/Runtime/Core/Public/HAL/MallocBinnedCommon.h b/Engine/Source/Runtime/Core/Public/HAL/MallocBinnedCommon.h
--- a/Engine/Source/Runtime/Core/Public/HAL/MallocBinnedCommon.h
+++ b/Engine/Source/Runtime/Core/Public/HAL/MallocBinnedCommon.h
@@ -57,16 +57,28 @@
 
 #ifndef UE_MBC_ALLOW_RUNTIME_TWEAKING
 #	define UE_MBC_ALLOW_RUNTIME_TWEAKING				0
+#elif UE_MBC_ALLOW_RUNTIME_TWEAKING
+#	ifndef UE_MBC_SUPPORT_RUNTIME_TLC_TOGGLE
+#		define UE_MBC_SUPPORT_RUNTIME_TLC_TOGGLE		1
+#	endif
 #endif
 
-#if UE_MBC_ALLOW_RUNTIME_TWEAKING
+#ifndef UE_MBC_SUPPORT_RUNTIME_TLC_TOGGLE
+#	define UE_MBC_SUPPORT_RUNTIME_TLC_TOGGLE			1
+#endif
+
+#if UE_MBC_SUPPORT_RUNTIME_TLC_TOGGLE
 	extern CORE_API int32 GMallocBinnedPerThreadCaches;
+#else
+#	define GMallocBinnedPerThreadCaches					UE_DEFAULT_GMallocBinnedPerThreadCaches
+#endif
+
+#if UE_MBC_ALLOW_RUNTIME_TWEAKING
 	extern CORE_API int32 GMallocBinnedBundleSize;
 	extern CORE_API int32 GMallocBinnedBundleCount;
 	extern CORE_API int32 GMallocBinnedAllocExtra;
 	extern CORE_API int32 GMallocBinnedMaxBundlesBeforeRecycle;
 #else
-#	define GMallocBinnedPerThreadCaches					UE_DEFAULT_GMallocBinnedPerThreadCaches
 #	define GMallocBinnedBundleSize						UE_DEFAULT_GMallocBinnedBundleSize
 #	define GMallocBinnedBundleCount						UE_DEFAULT_GMallocBinnedBundleCount
 #	define GMallocBinnedAllocExtra						UE_DEFAULT_GMallocBinnedAllocExtra
@@ -769,7 +781,7 @@ protected:
 	{
 		//NOALLOC_SCOPE_CYCLE_COUNTER(STAT_FMallocBinned_SetupTLSCachesOnCurrentThread);
 
-		if (!UE_MBC_ALLOW_RUNTIME_TWEAKING && !GMallocBinnedPerThreadCaches)
+		if (!GMallocBinnedPerThreadCaches)
 		{
 			return;
 		}
@@ -785,40 +797,37 @@ protected:
 	{
 		//NOALLOC_SCOPE_CYCLE_COUNTER(STAT_FMallocBinned_ClearTLSCachesOnCurrentThread);
 
-		if (!UE_MBC_ALLOW_RUNTIME_TWEAKING && !GMallocBinnedPerThreadCaches)
+		if (UE_MBC_SUPPORT_RUNTIME_TLC_TOGGLE || GMallocBinnedPerThreadCaches)
 		{
-			return;
+			// we need to flush cache even if it's disabled in case of run time tweaking disabled it, but it wasn't cleaned up
+			((AllocType*)this)->FlushCurrentThreadCacheInternal();
+			FPerThreadFreeBlockLists::ClearTLS();
 		}
-
-		((AllocType*)this)->FlushCurrentThreadCacheInternal();
-		FPerThreadFreeBlockLists::ClearTLS();
 	}
 
 	virtual void MarkTLSCachesAsUsedOnCurrentThread() override
 	{
 		//NOALLOC_SCOPE_CYCLE_COUNTER(STAT_FMallocBinned_MarkTLSCachesAsUsedOnCurrentThread);
 
-		if (!UE_MBC_ALLOW_RUNTIME_TWEAKING && !GMallocBinnedPerThreadCaches)
+		if (UE_MBC_SUPPORT_RUNTIME_TLC_TOGGLE || GMallocBinnedPerThreadCaches)
 		{
-			return;
+			// when UE_MBC_SUPPORT_RUNTIME_TLC_TOGGLE is set to 1 we still need to lock cache even if it is disabled, so we will clean it up properly
+			FPerThreadFreeBlockLists::LockTLS();
 		}
-
-		FPerThreadFreeBlockLists::LockTLS();
 	}
 
 	virtual void MarkTLSCachesAsUnusedOnCurrentThread() override
 	{
 		//NOALLOC_SCOPE_CYCLE_COUNTER(STAT_FMallocBinned_MarkTLSCachesAsUnusedOnCurrentThread);
 
-		if (!UE_MBC_ALLOW_RUNTIME_TWEAKING && !GMallocBinnedPerThreadCaches)
+		if (UE_MBC_SUPPORT_RUNTIME_TLC_TOGGLE || GMallocBinnedPerThreadCaches)
 		{
-			return;
-		}
+			// Will only flush if memory trimming epoch has been bumped while the thread was active.
+			const bool bNewEpochOnly = true;
+			((AllocType*)this)->FlushCurrentThreadCacheInternal(bNewEpochOnly);
 
-		// Will only flush if memory trimming epoch has been bumped while the thread was active.
-		const bool bNewEpochOnly = true;
-		((AllocType*)this)->FlushCurrentThreadCacheInternal(bNewEpochOnly);
-		FPerThreadFreeBlockLists::UnlockTLS();
+			FPerThreadFreeBlockLists::UnlockTLS();
+		}
 	}
 
 	inline SIZE_T QuantizeSizeCommon(SIZE_T Count, uint32 Alignment, const AllocType& Alloc) const
