Add FindTConsoleVariableDataString.

#Synced-CL 51472478

diff --git a/Engine/Source/Runtime/Core/Public/HAL/IConsoleManager.h b/Engine/Source/Runtime/Core/Public/HAL/IConsoleManager.h
--- a/Engine/Source/Runtime/Core/Public/HAL/IConsoleManager.h
+++ b/Engine/Source/Runtime/Core/Public/HAL/IConsoleManager.h
@@ -1212,6 +1212,18 @@ struct IConsoleManager
 		return P ? P->AsVariableFloat() : 0; 
 	}
 
+	/**
+	 * Find a typed console variable (faster access to the value, no virtual function call)
+	 * @param Name must not be 0
+	 * @return 0 if the object wasn't found
+	 */
+	TConsoleVariableData<FString>* FindTConsoleVariableDataString(const TCHAR* Name) const
+	{
+		IConsoleVariable* P = FindConsoleVariable(Name);
+
+		return P ? P->AsVariableString() : 0;
+	}
+
 
 	/**
 	 *  Iterate in O(n), not case sensitive, does not guarantee that UnregisterConsoleObject() will work in the loop
