Small fix to UpdateWeakReferenceInternal

#Synced-CL 51535387

diff --git a/Engine/Source/Runtime/Core/Public/Templates/SharedPointer.h b/Engine/Source/Runtime/Core/Public/Templates/SharedPointer.h
--- a/Engine/Source/Runtime/Core/Public/Templates/SharedPointer.h
+++ b/Engine/Source/Runtime/Core/Public/Templates/SharedPointer.h
@@ -1878,10 +1878,7 @@ private:
 	{
 		static_assert(OtherMode == Mode, "You cannot create a shared reference or pointer of one mode with a type which inherits TSharedFromThis of another mode.");
 
-		if (!ReferenceController)
-		{
-			ReferenceController = InReferenceController;
-		}
+		ReferenceController = InReferenceController;
 	}
 
 public:
