Made TOptional<T> have the correct copyability and movability traits of the underlying T.

diff --git a/Engine/Source/Runtime/Core/Public/Misc/IntrusiveUnsetOptionalState.h b/Engine/Source/Runtime/Core/Public/Misc/IntrusiveUnsetOptionalState.h
--- a/Engine/Source/Runtime/Core/Public/Misc/IntrusiveUnsetOptionalState.h
+++ b/Engine/Source/Runtime/Core/Public/Misc/IntrusiveUnsetOptionalState.h
@@ -5,11 +5,13 @@
 /**
  * Facilities to allow types to have an intrusive invalid state which can act as
  * TOptional's 'unset' state, saving space.  A class in such a state will only ever be
- * compared against FIntrusiveUnsetOptionalState or destructed.
+ * compared against FIntrusiveUnsetOptionalState or destructed.  It will be trivially
+ * copied if it's trivially copyable, otherwise its copy constructor and assignment operator
+ * won't be called in this state.
  *
  * A class should implement a constructor taking FIntrusiveUnsetOptionalState, and an 
  * equality comparison operator against FIntrusiveUnsetOptionalState, which will put a
-  class instance into the 'unset' state (in the case of the constructor) and allow testing
+ * class instance into the 'unset' state (in the case of the constructor) and allow testing
  * of its unset state.
  *
  * A public constexpr static data member of type bool called bHasIntrusiveUnsetOptionalState
diff --git a/Engine/Source/Runtime/Core/Public/Misc/Optional.h b/Engine/Source/Runtime/Core/Public/Misc/Optional.h
--- a/Engine/Source/Runtime/Core/Public/Misc/Optional.h
+++ b/Engine/Source/Runtime/Core/Public/Misc/Optional.h
@@ -72,6 +72,7 @@ public:
 	{
 	}
 	template <typename... ArgTypes>
+		requires (std::is_constructible_v<OptionalType, ArgTypes...>)
 	[[nodiscard]] explicit constexpr TOptional(EInPlace, ArgTypes&&... Args)
 		: TypedValue(Forward<ArgTypes>(Args)...)
 	{
@@ -121,6 +122,7 @@ public:
 
 	/** Construct an OptionalType with no value; i.e. unset */
 	[[nodiscard]] constexpr TOptional() requires (!bUsingIntrusiveUnsetState)
+		: UnsetValue()
 	{
 	}
 
@@ -136,6 +138,7 @@ public:
 
 	/** Copy/Move construction */
 	[[nodiscard]] TOptional(const TOptional& Other)
+		requires (std::is_copy_constructible_v<OptionalType> && !std::is_trivially_copy_constructible_v<OptionalType>)
 		: TOptional()
 	{
 		bool bLocalIsSet = Other.IsSet();
@@ -149,7 +152,11 @@ public:
 		}
 		// Default constructor initialized unset intrusive optional if necessary
 	}
+	[[nodiscard]] TOptional(const TOptional& Other) requires (std::is_copy_constructible_v<OptionalType> && std::is_trivially_copy_constructible_v<OptionalType>) = default;
+	[[nodiscard]] TOptional(const TOptional& Other) requires (!std::is_copy_constructible_v<OptionalType>) = delete;
+
 	[[nodiscard]] TOptional(TOptional&& Other)
+		requires (std::is_move_constructible_v<OptionalType> && !std::is_trivially_move_constructible_v<OptionalType>)
 		: TOptional()
 	{
 		bool bLocalIsSet = Other.IsSet();
@@ -163,8 +170,13 @@ public:
 		}
 		// Default constructor initialized unset intrusive optional if necessary
 	}
+	[[nodiscard]] TOptional(TOptional&& Other) requires (std::is_move_constructible_v<OptionalType> && std::is_trivially_move_constructible_v<OptionalType>) = default;
+	[[nodiscard]] TOptional(TOptional&& Other) requires (!std::is_move_constructible_v<OptionalType>) = delete;
 
 	TOptional& operator=(const TOptional& Other)
+		// Use is_constructible traits instead of is_assignable, because it uses Emplace, which always uses construction.
+		// If trivially copy constructible but not trivially copy assignable, handle this here too, so we don't accidentally copy assign.
+		requires (std::is_copy_constructible_v<OptionalType> && (!std::is_trivially_copy_constructible_v<OptionalType> || !std::is_trivially_copy_assignable_v<OptionalType>))
 	{
 		if (&Other != this)
 		{
@@ -179,7 +191,13 @@ public:
 		}
 		return *this;
 	}
+	TOptional& operator=(const TOptional& Other) requires (std::is_copy_assignable_v<OptionalType> && std::is_trivially_copy_constructible_v<OptionalType> && std::is_trivially_copy_assignable_v<OptionalType>) = default;
+	TOptional& operator=(const TOptional& Other) requires (!std::is_copy_constructible_v<OptionalType>) = delete;
+
 	TOptional& operator=(TOptional&& Other)
+		// Use is_constructible traits instead of is_assignable, because it uses Emplace, which always uses construction.
+		// If trivially move constructible but not trivially move assignable, handle this here too, so we don't accidentally move assign.
+		requires (std::is_move_constructible_v<OptionalType> && (!std::is_trivially_move_constructible_v<OptionalType> || !std::is_trivially_move_assignable_v<OptionalType>))
 	{
 		if (&Other != this)
 		{
@@ -194,6 +212,8 @@ public:
 		}
 		return *this;
 	}
+	TOptional& operator=(TOptional&& Other) requires (std::is_move_assignable_v<OptionalType> && std::is_trivially_move_constructible_v<OptionalType> && std::is_trivially_move_assignable_v<OptionalType>) = default;
+	TOptional& operator=(TOptional&& Other) requires (!std::is_move_constructible_v<OptionalType>) = delete;
 
 	TOptional& operator=(const OptionalType& InValue)
 	{
@@ -223,12 +243,14 @@ public:
 			}
 			else
 			{
+				this->UnsetValue = {};
 				this->bIsSet = false;
 			}
 		}
 	}
 
 	template <typename... ArgsType>
+		requires (std::is_constructible_v<OptionalType, ArgsType...>)
 	OptionalType& Emplace(ArgsType&&... Args)
 	{
 		// Destroy the member in-place before replacing it - a bit nasty, but it'll work since we don't support exceptions
@@ -406,9 +428,20 @@ private:
 		DestructItem(std::addressof(this->TypedValue));
 	}
 
+	// This dummy type exists to prevent the OptionalType from being constructed when
+	// the union needs to be constructed, when in constexpr contexts, the union needs
+	// to be constructed to something.
+	struct FUnsetType
+	{
+		[[nodiscard]] UE_REWRITE constexpr FUnsetType()
+		{
+		}
+	};
+
 	// If we can have an explicitly conditional trivial destructor, we can declare our members inline and avoid the base class
 	union
 	{
+		FUnsetType UnsetValue;
 		OptionalType TypedValue;
 	};
 	// Set flag takes up no space if bIsUsingIntrusiveUnsetState, otherwise is zero-initialized to false
diff --git a/Engine/Source/Runtime/Core/Tests/Misc/OptionalTest.cpp b/Engine/Source/Runtime/Core/Tests/Misc/OptionalTest.cpp
--- a/Engine/Source/Runtime/Core/Tests/Misc/OptionalTest.cpp
+++ b/Engine/Source/Runtime/Core/Tests/Misc/OptionalTest.cpp
@@ -5,6 +5,7 @@
 #if WITH_TESTS
 
 #include "CoreTypes.h"
+#include "Containers/StaticArray.h"
 #include "Containers/UnrealString.h"
 #include "Internationalization/Text.h"
 #include "Math/Vector.h"
@@ -144,4 +145,156 @@ TEST_CASE_NAMED(FTOptionalTest_Intrusive, "System::Core::TOptional::Intrusive",
 	CHECK_MESSAGE("Unset optional is not destroyed", FIntrusiveDestructible::DestructorCount == 0);
 }
 
+// Test that trivial copyability is transitive
+static_assert(std::is_trivially_copyable_v<TOptional<bool>>                      == std::is_trivially_copyable_v<bool>);
+static_assert(std::is_trivially_copyable_v<TOptional<int32>>                     == std::is_trivially_copyable_v<int32>);
+static_assert(std::is_trivially_copyable_v<TOptional<double>>                    == std::is_trivially_copyable_v<double>);
+static_assert(std::is_trivially_copyable_v<TOptional<FName>>                     == std::is_trivially_copyable_v<FName>);
+static_assert(std::is_trivially_copyable_v<TOptional<FString*>>                  == std::is_trivially_copyable_v<FString*>);
+static_assert(std::is_trivially_copyable_v<TOptional<TSharedPtr<int32>>>         == std::is_trivially_copyable_v<TSharedPtr<int32>>);
+static_assert(std::is_trivially_copyable_v<TOptional<FString>>                   == std::is_trivially_copyable_v<FString>);
+static_assert(std::is_trivially_copyable_v<TOptional<FVector>>                   == std::is_trivially_copyable_v<FVector>);
+static_assert(std::is_trivially_copyable_v<TOptional<TArray<int32>>>             == std::is_trivially_copyable_v<TArray<int32>>);
+static_assert(std::is_trivially_copyable_v<TOptional<TStaticArray<int32, 20>>>   == std::is_trivially_copyable_v<TStaticArray<int32, 20>>);
+static_assert(std::is_trivially_copyable_v<TOptional<TStaticArray<FString, 20>>> == std::is_trivially_copyable_v<TStaticArray<FString, 20>>);
+
+// Test that copyability traits are as expected.
+
+// An optional of a trivial type should have trivial operations
+namespace TestTrivialOptionalType
+{
+	using OptionalType = int32;
+
+	static_assert(std::is_trivially_copy_constructible_v<TOptional<OptionalType>>);
+	static_assert(std::is_trivially_copy_assignable_v<TOptional<OptionalType>>);
+
+	static_assert(std::is_trivially_move_constructible_v<TOptional<OptionalType>>);
+	static_assert(std::is_trivially_move_assignable_v<TOptional<OptionalType>>);
+
+	static_assert(std::is_trivially_destructible_v<TOptional<OptionalType>>);
+}
+// An optional of a movable and copyable type should be movable and copyable but not have have trivial operations.
+namespace TestMovableAndCopyableOptionalType
+{
+	using OptionalType = FString;
+
+	static_assert(std::is_copy_constructible_v<TOptional<OptionalType>>);
+	static_assert(std::is_copy_assignable_v<TOptional<OptionalType>>);
+	static_assert(!std::is_trivially_copy_constructible_v<TOptional<OptionalType>>);
+	static_assert(!std::is_trivially_copy_assignable_v<TOptional<OptionalType>>);
+
+	static_assert(std::is_move_constructible_v<TOptional<OptionalType>>);
+	static_assert(std::is_move_assignable_v<TOptional<OptionalType>>);
+	static_assert(!std::is_trivially_move_constructible_v<TOptional<OptionalType>>);
+	static_assert(!std::is_trivially_move_assignable_v<TOptional<OptionalType>>);
+
+	static_assert(!std::is_trivially_destructible_v<TOptional<OptionalType>>);
+}
+// An optional of a movable-only type should be movable but not copyable.
+namespace TestMovableOnlyOptionalType
+{
+	using OptionalType = TUniquePtr<int>;
+
+	static_assert(!std::is_copy_constructible_v<TOptional<OptionalType>>);
+	static_assert(!std::is_copy_assignable_v<TOptional<OptionalType>>);
+
+	static_assert(std::is_move_constructible_v<TOptional<OptionalType>>);
+	static_assert(std::is_move_assignable_v<TOptional<OptionalType>>);
+	static_assert(!std::is_trivially_move_constructible_v<TOptional<OptionalType>>);
+	static_assert(!std::is_trivially_move_assignable_v<TOptional<OptionalType>>);
+
+	static_assert(!std::is_trivially_destructible_v<TOptional<OptionalType>>);
+}
+// An optional of a type which is neither copyable nor movable, but trivially destructible should be neither copyable nor movable, but trivially destructible.
+namespace TestNonMovableTriviallyDestructibleType
+{
+	struct OptionalType
+	{
+		OptionalType(OptionalType&&) = delete;
+		OptionalType(const OptionalType&) = delete;
+		OptionalType& operator=(OptionalType&&) = delete;
+		OptionalType& operator=(const OptionalType&) = delete;
+		~OptionalType() = default;
+	};
+
+	static_assert(!std::is_copy_constructible_v<TOptional<OptionalType>>);
+	static_assert(!std::is_copy_assignable_v<TOptional<OptionalType>>);
+
+	static_assert(!std::is_move_constructible_v<TOptional<OptionalType>>);
+	static_assert(!std::is_move_assignable_v<TOptional<OptionalType>>);
+
+	static_assert(std::is_trivially_destructible_v<TOptional<OptionalType>>);
+}
+// An optional of a type which is neither copyable, movable nor trivially destructible should be neither copyable, movable nor trivially destructible.
+namespace TestNonMovableNonTriviallyDestructibleType
+{
+	struct OptionalType
+	{
+		OptionalType(OptionalType&&) = delete;
+		OptionalType(const OptionalType&) = delete;
+		OptionalType& operator=(OptionalType&&) = delete;
+		OptionalType& operator=(const OptionalType&) = delete;
+		~OptionalType();
+	};
+
+	static_assert(!std::is_copy_constructible_v<TOptional<OptionalType>>);
+	static_assert(!std::is_copy_assignable_v<TOptional<OptionalType>>);
+
+	static_assert(!std::is_move_constructible_v<TOptional<OptionalType>>);
+	static_assert(!std::is_move_assignable_v<TOptional<OptionalType>>);
+
+	static_assert(!std::is_trivially_move_constructible_v<TOptional<OptionalType>>);
+	static_assert(!std::is_trivially_move_assignable_v<TOptional<OptionalType>>);
+
+	static_assert(!std::is_trivially_destructible_v<OptionalType>);
+	static_assert(!std::is_trivially_destructible_v<TOptional<OptionalType>>);
+}
+// An optional of a type which is only copy constructible but not copy assignable should be copyable *and* move assignable, because TOptional assignment destructs and reconstructs.
+// However, it is not trivially copy assignable even if the type is trivially copy constructible.
+namespace TestTriviallyCopyConstructibleOnlyOptionalType
+{
+	struct OptionalType
+	{
+		OptionalType(OptionalType&&) = default;
+		OptionalType(const OptionalType&) = default;
+		OptionalType& operator=(OptionalType&&) = delete;
+		OptionalType& operator=(const OptionalType&) = delete;
+		~OptionalType() = default;
+	};
+
+	static_assert(std::is_copy_constructible_v<TOptional<OptionalType>>);
+	static_assert(std::is_copy_assignable_v<TOptional<OptionalType>>);
+	static_assert(std::is_trivially_copy_constructible_v<TOptional<OptionalType>>);
+	static_assert(!std::is_trivially_copy_assignable_v<TOptional<OptionalType>>);
+
+	static_assert(std::is_move_constructible_v<TOptional<OptionalType>>);
+	static_assert(std::is_move_assignable_v<TOptional<OptionalType>>);
+	static_assert(std::is_trivially_move_constructible_v<TOptional<OptionalType>>);
+	static_assert(!std::is_trivially_move_assignable_v<TOptional<OptionalType>>);
+
+	static_assert(std::is_trivially_destructible_v<TOptional<OptionalType>>);
+}
+// Same as above but only movable
+namespace TestTriviallyMoveConstructibleOnlyAndDestructibleOptionalType
+{
+	struct OptionalType
+	{
+		OptionalType(OptionalType&&) = default;
+		OptionalType(const OptionalType&) = delete;
+		OptionalType& operator=(OptionalType&&) = delete;
+		OptionalType& operator=(const OptionalType&) = delete;
+		~OptionalType() = default;
+	};
+
+	static_assert(!std::is_copy_constructible_v<TOptional<OptionalType>>);
+	static_assert(!std::is_copy_assignable_v<TOptional<OptionalType>>);
+
+	static_assert(std::is_move_constructible_v<TOptional<OptionalType>>);
+	static_assert(std::is_move_assignable_v<TOptional<OptionalType>>);
+	static_assert(std::is_trivially_move_constructible_v<TOptional<OptionalType>>);
+	static_assert(!std::is_trivially_move_assignable_v<TOptional<OptionalType>>);
+
+	static_assert(std::is_trivially_destructible_v<TOptional<OptionalType>>);
+}
+
 #endif // WITH_TESTS
\ No newline at end of file
