[QueuedThreadPoolWrapper]
- Fix missing release when abandoning work during destroy

#Synced-CL 51539428

diff --git a/Engine/Source/Runtime/Core/Private/Misc/QueuedThreadPoolWrapper.cpp b/Engine/Source/Runtime/Core/Private/Misc/QueuedThreadPoolWrapper.cpp
--- a/Engine/Source/Runtime/Core/Private/Misc/QueuedThreadPoolWrapper.cpp
+++ b/Engine/Source/Runtime/Core/Private/Misc/QueuedThreadPoolWrapper.cpp
@@ -69,6 +69,10 @@ void FQueuedThreadPoolWrapper::Destroy()
 			if (WrappedQueuedThreadPool->RetractQueuedWork(Work))
 			{
 				Work->GetInnerWork()->Abandon();
+				// Decrement the ref added by Assign() before releasing.
+				// Matches the pattern in RetractQueuedWork() and TryRetractWorkNoLock().
+				check(Work->NumRefs.GetValue() == 1);
+				Work->NumRefs.Decrement();
 				ReleaseWorkNoLock(Work);
 			}
 		}
