[Insights] Disable unwind tables method for callstack tracing for Clang + exceptions

Clang uses physical frame pointers for exception handling in complex cases (where destructors need to be called), while Microsoft uses "state tables" and eliminates the frame pointers. Since some fundamental methods (e.g. FMemory::Allocate, etc) are marked as exception handlers, this marks them with the UWOP_SET_FPREG flag. Whenever stack walking encounters such a function it stops unwinding, resulting in a truncated stack, which makes callstacks unsuitable for memory tracing.

#Synced-CL 51335054

diff --git a/Engine/Source/Runtime/Core/Private/ProfilingDebugging/Microsoft/WindowsCallstackTrace.cpp b/Engine/Source/Runtime/Core/Private/ProfilingDebugging/Microsoft/WindowsCallstackTrace.cpp
--- a/Engine/Source/Runtime/Core/Private/ProfilingDebugging/Microsoft/WindowsCallstackTrace.cpp
+++ b/Engine/Source/Runtime/Core/Private/ProfilingDebugging/Microsoft/WindowsCallstackTrace.cpp
@@ -45,7 +45,14 @@ static bool GModulesAreInitialized = false;
 // stack walking. In some cases this is not suitable, and we then fall back
 // to the standard stack walking implementation.
 #if !defined(UE_CALLSTACK_TRACE_USE_UNWIND_TABLES)
-	#define UE_CALLSTACK_TRACE_USE_UNWIND_TABLES 1
+	// Clang uses physical frame pointer for exception handling, causing
+	// many functions (e.g. FMemory::Allocate) to be marked with UWOP_SET_FPREG
+	// making callstacks unsuitable for allocation tracing.
+	#if PLATFORM_COMPILER_CLANG && !PLATFORM_EXCEPTIONS_DISABLED
+		#define UE_CALLSTACK_TRACE_USE_UNWIND_TABLES 0
+	#else
+		#define UE_CALLSTACK_TRACE_USE_UNWIND_TABLES 1
+	#endif
 #endif
 
 // ARM64 and ARM64EC must use fallback paths. Unwind tables for those architectures have not been implemented.
