Rename some private stats types to avoid collision with platform SDKs.

#Synced-CL 51330535

diff --git a/Engine/Source/Runtime/Core/Public/Stats/StatsSystemTypes.h b/Engine/Source/Runtime/Core/Public/Stats/StatsSystemTypes.h
--- a/Engine/Source/Runtime/Core/Public/Stats/StatsSystemTypes.h
+++ b/Engine/Source/Runtime/Core/Public/Stats/StatsSystemTypes.h
@@ -1487,9 +1487,9 @@ class FCycleCounter
 	/** Event types emitted */
 	enum
 	{
-		NamedEvent = 1 << 0,
-		TraceEvent = 1 << 1,
-		ThreadStatsEvent = 1 << 2,
+		NamedEventFlag = 1 << 0,
+		TraceEventFlag = 1 << 1,
+		ThreadStatsEventFlag = 1 << 2,
 	};
 
 	/** Name of the stat, usually a short name **/
@@ -1520,14 +1520,14 @@ public:
 #else
 			FPlatformMisc::BeginNamedEvent(FColor(0), InStatId.GetStatDescriptionWIDE());
 #endif
-			EmittedEvent |= NamedEvent;
+			EmittedEvent |= NamedEventFlag;
 
 #if CPUPROFILERTRACE_ENABLED
 			if (UE_TRACE_CHANNELEXPR_IS_ENABLED(CpuChannel))
 			{
 				FName StatName = MinimalNameToName(StatMinimalName);
 				FCpuProfilerTrace::OutputBeginDynamicEventWithId(StatName, InStatId.GetStatDescriptionWIDE(), SourceLocation.GetFileName(), SourceLocation.GetLine());
-				EmittedEvent |= TraceEvent;
+				EmittedEvent |= TraceEventFlag;
 			}
 #endif
 		}
@@ -1537,7 +1537,7 @@ public:
 			FName StatName = MinimalNameToName(StatMinimalName);
 			StatId = StatName;
 			FThreadStats::AddMessage(StatName, EStatOperation::CycleScopeStart);
-			EmittedEvent |= ThreadStatsEvent;
+			EmittedEvent |= ThreadStatsEventFlag;
 		}
 	}
 
@@ -1554,7 +1554,7 @@ public:
 		if (UE_TRACE_CHANNELEXPR_IS_ENABLED(CpuChannel))
 		{
 			FCpuProfilerTrace::OutputBeginDynamicEvent(Name);
-			EmittedEvent |= TraceEvent;
+			EmittedEvent |= TraceEventFlag;
 		}
 #endif
 	}
@@ -1566,7 +1566,7 @@ public:
 		if (UE_TRACE_CHANNELEXPR_IS_ENABLED(CpuChannel))
 		{
 			FCpuProfilerTrace::OutputBeginDynamicEventWithId(Name, Desc);
-			EmittedEvent |= TraceEvent;
+			EmittedEvent |= TraceEventFlag;
 		}
 #endif
 	}
@@ -1577,19 +1577,19 @@ public:
 	UE_AUTORTFM_ALWAYS_OPEN
 	inline void Stop()
 	{
-		if (EmittedEvent & NamedEvent)
+		if (EmittedEvent & NamedEventFlag)
 		{
 			FPlatformMisc::EndNamedEvent();
 		}
 
 #if CPUPROFILERTRACE_ENABLED
-		if (EmittedEvent & TraceEvent)
+		if (EmittedEvent & TraceEventFlag)
 		{
 			FCpuProfilerTrace::OutputEndEvent();
 		}
 #endif
 
-		if (EmittedEvent & ThreadStatsEvent)
+		if (EmittedEvent & ThreadStatsEventFlag)
 		{
 			FThreadStats::AddMessage(StatId, EStatOperation::CycleScopeEnd);
 		}
