Replaced TLS slots with thread_local for per thread caches in MB2 & MB3 as it has better performance

#Synced-CL 51465518

diff --git a/Engine/Source/Runtime/Core/Private/HAL/MallocBinned2.cpp b/Engine/Source/Runtime/Core/Private/HAL/MallocBinned2.cpp
--- a/Engine/Source/Runtime/Core/Private/HAL/MallocBinned2.cpp
+++ b/Engine/Source/Runtime/Core/Private/HAL/MallocBinned2.cpp
@@ -100,6 +100,9 @@ FMallocBinned2* FMallocBinned2::MallocBinned2 = nullptr;
 // Mapping of sizes to small table indices
 uint8 FMallocBinned2::MemSizeToPoolIndex[SIZE_TO_POOL_INDEX_NUM] = { 0 };
 
+template <>
+thread_local FMallocBinned2::FPerThreadFreeBlockLists* FMallocBinned2::FPerThreadFreeBlockLists::ThreadSingleton = nullptr;
+
 
 static FORCEINLINE bool IsSupportedSize(SIZE_T Size)
 {
@@ -610,6 +613,57 @@ void* FMallocBinned2::MallocExternalSmall(SIZE_T Size, uint32 Alignment)
 	return Result;
 }
 
+void* FMallocBinned2::MallocNoTLS(SIZE_T Size, uint32 Alignment)
+{
+	// Only allocate from the small pools if the size is small enough and the alignment isn't crazy large.
+	// With large alignments, we'll waste a lot of memory allocating an entire page, but such alignments are highly unlikely in practice.
+	bool bUseSmallPool = UseSmallAlloc(Size, Alignment);
+	if (!bUseSmallPool)
+	{
+		bUseSmallPool = PromoteToLargerBin(Size, Alignment, *this);
+	}
+
+	if (bUseSmallPool)
+	{
+		const uint32 PoolIndex = BoundSizeToPoolIndex(Size, MemSizeToPoolIndex);
+
+		NOALLOC_SCOPE_CYCLE_COUNTER(STAT_FMallocBinned2_MallocExternalSmall);
+
+		// Allocate from small object pool.
+		FPoolTable& Table = SmallPoolTables[PoolIndex];
+
+		UE::TUniqueLock Lock(Table.Mutex);
+
+		FPoolInfo* Pool;
+		if (!Table.ActivePools.IsEmpty())
+		{
+			Pool = &Table.ActivePools.GetFrontPool();
+		}
+		else
+		{
+			Pool = &Table.ActivePools.PushNewPoolToFront(*this, Table, PoolIndex);
+		}
+
+		void* Result = Pool->AllocateBin();
+#if UE_MB2_ALLOCATOR_STATS
+		++Table.TotalUsedBins;
+		AllocatedSmallPoolMemory.fetch_add(Table.BinSize, std::memory_order_relaxed);
+#endif
+
+#if UE_MBC_LIGHTWEIGHT_BIN_CALLSTACK_TRACKER
+		TrackBinAllocation(Table.BinSize, Table.TotalAllocatedMem);
+#endif
+
+		if (!Pool->HasFreeBin())
+		{
+			Table.ExhaustedPools.LinkToFront(Pool);
+		}
+
+		return Result;
+	}
+	return MallocExternalLarge(Size, Alignment);
+}
+
 void* FMallocBinned2::MallocExternalLarge(SIZE_T Size, uint32 Alignment)
 {
 	Alignment = FMath::Max<uint32>(Alignment, UE_MBC_STANDARD_ALIGNMENT);
@@ -824,6 +878,35 @@ void FMallocBinned2::Free(void* Ptr)
 	}
 }
 
+void FMallocBinned2::FreeNoTLS(void* Ptr)
+{
+	NOALLOC_SCOPE_CYCLE_COUNTER(STAT_FMallocBinned2_FreeExternal);
+
+	if (!IsOSAllocation(Ptr))
+	{
+		check(Ptr); // null is 64k aligned so we should not be here
+		FFreeBlock* BasePtr = GetPoolHeaderFromPointer(Ptr);
+		CanaryTest(BasePtr);
+
+#if BINNED2_FORK_SUPPORT
+		if (BasePtr->CanaryAndForkState != CurrentCanary)
+		{
+			// This page was allocated before we forked so we want to avoid dirtying it by writing a linked list into it 
+			return;
+		}
+#endif
+
+		FBundleNode* BundlesToRecycle = (FBundleNode*)Ptr;
+		BundlesToRecycle->SetNextNodeInCurrentBundle(nullptr);
+
+		Private::FreeBundles(*this, BundlesToRecycle, BasePtr->PoolIndex);
+	}
+	else if (Ptr)
+	{
+		Free(Ptr);
+	}
+}
+
 void FMallocBinned2::FPoolList::ValidateActivePools() const
 {
 	for (const FPoolInfo* const* PoolPtr = &Front; *PoolPtr; PoolPtr = &(*PoolPtr)->Next)
diff --git a/Engine/Source/Runtime/Core/Private/HAL/MallocBinned3.cpp b/Engine/Source/Runtime/Core/Private/HAL/MallocBinned3.cpp
--- a/Engine/Source/Runtime/Core/Private/HAL/MallocBinned3.cpp
+++ b/Engine/Source/Runtime/Core/Private/HAL/MallocBinned3.cpp
@@ -59,6 +59,9 @@ FMallocBinned3* FMallocBinned3::MallocBinned3 = nullptr;
 // Mapping of sizes to small table indices
 uint8 FMallocBinned3::MemSizeToPoolIndex[SIZE_TO_POOL_INDEX_NUM] = { 0 };
 
+template <>
+thread_local FMallocBinned3::FPerThreadFreeBlockLists* FMallocBinned3::FPerThreadFreeBlockLists::ThreadSingleton = nullptr;
+
 template <class T>
 static void SetCanaryInternal(T& PoolInfo, typename T::ECanary ShouldBe, bool bPreexisting, bool bGuarnteedToBeNew)
 {
@@ -762,6 +765,75 @@ void* FMallocBinned3::Malloc(SIZE_T Size, uint32 Alignment)
 	return Result;
 }
 
+void* FMallocBinned3::MallocNoTLS(SIZE_T Size, uint32 Alignment)
+{
+	bool UsePools = (Size <= UE_MB3_MAX_SMALL_POOL_SIZE);
+
+	check(FMath::IsPowerOfTwo(Alignment));
+
+	if (UNLIKELY(UsePools && (Alignment > UE_MBC_STANDARD_ALIGNMENT)))
+	{
+		// check if allocations that require alignment larger than UE_MBC_STANDARD_ALIGNMENT can be promoted to a bin with a natural alignment that matches
+		// i.e. 16 bytes allocation with 128 bytes alignment can be promoted to 128 bytes bin
+		// this will save us a lot of memory as otherwise allocations will be promoted to OS allocs that are at least 4 KB large, depending on a platform
+		UsePools = PromoteToLargerBin(Size, Alignment, *this);
+	}
+
+	if (UsePools)
+	{
+		const uint32 PoolIndex = BoundSizeToPoolIndex(Size, MemSizeToPoolIndex);
+
+		NOALLOC_SCOPE_CYCLE_COUNTER(STAT_FMallocBinned3_MallocExternalSmall);
+
+		// Allocate from small object pool.
+		FPoolTable& Table = SmallPoolTables[PoolIndex];
+
+		UE::TUniqueLock Lock(Table.Mutex);
+
+		uint32 BlockIndex = MAX_uint32;
+		FPoolInfoSmall* Pool = GetFrontPool(Table, PoolIndex, BlockIndex);
+		if (!Pool)
+		{
+			Pool = PushNewPoolToFront(Table, Table.BinSize, PoolIndex, BlockIndex);
+
+			// Indicates that we run out of reserved virtual memory for the pool (UE_MB3_MAX_MEMORY_PER_POOL_SIZE) for this bin type
+			if (!Pool)
+			{
+				if ((PoolIndex + 1) < UE_MB3_SMALL_POOL_COUNT)
+				{
+					return FMallocBinned3::Malloc(SmallPoolTables[PoolIndex + 1].BinSize, Alignment);
+				}
+				else
+				{
+					return FMallocBinned3::Malloc(UE_MB3_MAX_SMALL_POOL_SIZE + 1, Alignment);
+				}
+			}
+		}
+
+		uint8* BlockPtr = BlockPointerFromIndecies(PoolIndex, BlockIndex, Table.BlockSize);
+
+		void* Result = Pool->AllocateBin(BlockPtr, Table.BinSize);
+#if UE_MB3_ALLOCATOR_STATS
+		++Table.TotalUsedBins;
+		Table.HeadEndAlloc(Size);
+		AllocatedSmallPoolMemory.fetch_add(Table.BinSize, std::memory_order_relaxed);
+#endif
+
+#if UE_MBC_LIGHTWEIGHT_BIN_CALLSTACK_TRACKER
+		TrackBinAllocation(Table.BinSize, Table.TotalAllocatedMem);
+#endif
+
+		if (!Pool->HasFreeBin())
+		{
+			Table.BlocksExhaustedBits.AllocBit(BlockIndex);
+		}
+
+		return Result;
+	}
+
+	return Malloc(Size, Alignment);
+}
+
 void* FMallocBinned3::Realloc(void* Ptr, SIZE_T NewSize, uint32 Alignment)
 {
 	const uint64 PoolIndex = PoolIndexFromPtr(Ptr);
@@ -930,6 +1002,32 @@ void FMallocBinned3::FreeExternal(void* Ptr, uint64 PoolIndex)
 	}
 }
 
+void FMallocBinned3::FreeNoTLS(void* Ptr)
+{
+	NOALLOC_SCOPE_CYCLE_COUNTER(STAT_FMallocBinned3_FreeExternal);
+
+	const uint64 PoolIndex = PoolIndexFromPtr(Ptr);
+	if (PoolIndex < UE_MB3_SMALL_POOL_COUNT)
+	{
+		check(Ptr); // null is an OS allocation because it will not fall in our VM block
+		const uint32 BinSize = PoolIndexToBinSize(PoolIndex);
+
+		FBundleNode* BundlesToRecycle = (FBundleNode*)Ptr;
+		BundlesToRecycle->SetNextNodeInCurrentBundle(nullptr);
+
+		Private::FreeBundles(*this, BundlesToRecycle, PoolIndex);
+#if UE_MB3_ALLOCATOR_STATS
+		SmallPoolTables[PoolIndex].HeadEndFree();
+		// lists track their own stat track them instead in the global stat if we don't have lists
+		AllocatedSmallPoolMemory.fetch_sub((int64)BinSize, std::memory_order_relaxed);
+#endif
+	}
+	else if (Ptr)
+	{
+		FreeExternal(Ptr, PoolIndex);
+	}
+}
+
 bool FMallocBinned3::ValidateHeap()
 {
 	// Not implemented
diff --git a/Engine/Source/Runtime/Core/Private/HAL/MallocBinnedCommon.cpp b/Engine/Source/Runtime/Core/Private/HAL/MallocBinnedCommon.cpp
--- a/Engine/Source/Runtime/Core/Private/HAL/MallocBinnedCommon.cpp
+++ b/Engine/Source/Runtime/Core/Private/HAL/MallocBinnedCommon.cpp
@@ -590,7 +590,6 @@ static FAutoConsoleVariableRef GMallocBinnedFlushRegisteredThreadCachesOnOneThre
 	);
 #endif	//~UE_MBC_ALLOW_RUNTIME_TWEAKING
 
-uint32 FMallocBinnedCommonBase::BinnedTlsSlot = FPlatformTLS::InvalidTlsSlot;
 #if UE_MBC_ALLOCATOR_STATS
 	std::atomic<int64> FMallocBinnedCommonBase::TLSMemory(0);
 	std::atomic<int64> FMallocBinnedCommonBase::ConsolidatedMemory(0);
diff --git a/Engine/Source/Runtime/Core/Public/HAL/MallocBinned2.h b/Engine/Source/Runtime/Core/Public/HAL/MallocBinned2.h
--- a/Engine/Source/Runtime/Core/Public/HAL/MallocBinned2.h
+++ b/Engine/Source/Runtime/Core/Public/HAL/MallocBinned2.h
@@ -242,12 +242,16 @@ public:
 	UE_AUTORTFM_NOAUTORTFM
 	virtual void* Malloc(SIZE_T Size, uint32 Alignment) override;
 
+	void* MallocNoTLS(SIZE_T Size, uint32 Alignment);
+
 	UE_AUTORTFM_NOAUTORTFM
 	virtual void* Realloc(void* Ptr, SIZE_T NewSize, uint32 Alignment) override;
 
 	UE_AUTORTFM_NOAUTORTFM
 	virtual void Free(void* Ptr) override;
 
+	void FreeNoTLS(void* Ptr);
+
 	inline bool GetSmallAllocationSize(void* Ptr, SIZE_T& SizeOut) const
 	{
 		if (!IsOSAllocation(Ptr))
diff --git a/Engine/Source/Runtime/Core/Public/HAL/MallocBinned3.h b/Engine/Source/Runtime/Core/Public/HAL/MallocBinned3.h
--- a/Engine/Source/Runtime/Core/Public/HAL/MallocBinned3.h
+++ b/Engine/Source/Runtime/Core/Public/HAL/MallocBinned3.h
@@ -233,6 +233,8 @@ public:
 	UE_AUTORTFM_NOAUTORTFM
 	virtual void* Malloc(SIZE_T Size, uint32 Alignment) override;
 
+	void* MallocNoTLS(SIZE_T Size, uint32 Alignment);
+
 	UE_AUTORTFM_NOAUTORTFM
 	virtual void* Realloc(void* Ptr, SIZE_T NewSize, uint32 Alignment) override;
 
@@ -242,6 +244,8 @@ public:
 		FreeExternal(Ptr, PoolIndexFromPtr(Ptr));
 	}
 
+	void FreeNoTLS(void* Ptr);
+
 	inline bool GetSmallAllocationSize(void* Ptr, SIZE_T& SizeOut) const
 	{
 		const uint64 PoolIndex = PoolIndexFromPtr(Ptr);
diff --git a/Engine/Source/Runtime/Core/Public/HAL/MallocBinnedCommon.h b/Engine/Source/Runtime/Core/Public/HAL/MallocBinnedCommon.h
--- a/Engine/Source/Runtime/Core/Public/HAL/MallocBinnedCommon.h
+++ b/Engine/Source/Runtime/Core/Public/HAL/MallocBinnedCommon.h
@@ -410,8 +410,6 @@ protected:
 	static uint32 OsAllocationGranularity;
 	UE::FPlatformRecursiveMutex ExternalAllocMutex;
 
-	static CORE_API uint32 BinnedTlsSlot;
-
 	std::atomic<uint64> MemoryTrimEpoch{ 0 };
 
 #if UE_MBC_ALLOCATOR_STATS
@@ -501,9 +499,10 @@ protected:
 
 	struct FPerThreadFreeBlockLists
 	{
+		static thread_local FPerThreadFreeBlockLists* ThreadSingleton;
+
 		inline static FPerThreadFreeBlockLists* Get() TSAN_SAFE
 		{
-			FPerThreadFreeBlockLists* ThreadSingleton = FPlatformTLS::IsValidTlsSlot(BinnedTlsSlot) ? (FPerThreadFreeBlockLists*)FPlatformTLS::GetTlsValue(BinnedTlsSlot) : nullptr;
 			// If the current thread doesn't have the Lock, we can't return the TLS cache for being used on the current thread as we risk racing with another thread doing trimming.
 			// This can only happen in such a scenario.
 			//
@@ -519,8 +518,6 @@ protected:
 
 		static void SetTLS()
 		{
-			check(FPlatformTLS::IsValidTlsSlot(BinnedTlsSlot));
-			FPerThreadFreeBlockLists* ThreadSingleton = (FPerThreadFreeBlockLists*)FPlatformTLS::GetTlsValue(BinnedTlsSlot);
 			if (!ThreadSingleton)
 			{
 				const int64 TLSSize = Align(sizeof(FPerThreadFreeBlockLists), AllocType::OsAllocationGranularity);
@@ -530,14 +527,13 @@ protected:
 				verify(ThreadSingleton);
 				ThreadSingleton->bLockedByOwnerThread = true;
 				ThreadSingleton->Lock();
-				FPlatformTLS::SetTlsValue(BinnedTlsSlot, ThreadSingleton);
+
 				RegisterThreadFreeBlockLists(ThreadSingleton);
 			}
 		}
 
 		static void UnlockTLS()
 		{
-			FPerThreadFreeBlockLists* ThreadSingleton = (FPerThreadFreeBlockLists*)FPlatformTLS::GetTlsValue(BinnedTlsSlot);
 			if (ThreadSingleton)
 			{
 				ThreadSingleton->bLockedByOwnerThread = false;
@@ -547,7 +543,6 @@ protected:
 
 		static void LockTLS()
 		{
-			FPerThreadFreeBlockLists* ThreadSingleton = (FPerThreadFreeBlockLists*)FPlatformTLS::GetTlsValue(BinnedTlsSlot);
 			if (ThreadSingleton)
 			{
 				ThreadSingleton->Lock();
@@ -557,8 +552,6 @@ protected:
 
 		static void ClearTLS()
 		{
-			check(FPlatformTLS::IsValidTlsSlot(BinnedTlsSlot));
-			FPerThreadFreeBlockLists* ThreadSingleton = (FPerThreadFreeBlockLists*)FPlatformTLS::GetTlsValue(BinnedTlsSlot);
 			if (ThreadSingleton)
 			{
 				const int64 TLSSize = Align(sizeof(FPerThreadFreeBlockLists), AllocType::OsAllocationGranularity);
@@ -570,8 +563,9 @@ protected:
 				ThreadSingleton->~FPerThreadFreeBlockLists();
 
 				AllocType::FreeMetaDataMemory(ThreadSingleton, TLSSize);
+
+				ThreadSingleton = nullptr;
 			}
-			FPlatformTLS::SetTlsValue(BinnedTlsSlot, nullptr);
 		}
 
 		UE_FORCEINLINE_HINT void* Malloc(uint32 InPoolIndex)
@@ -785,11 +779,7 @@ protected:
 		{
 			return;
 		}
-		if (!FPlatformTLS::IsValidTlsSlot(BinnedTlsSlot))
-		{
-			BinnedTlsSlot = FPlatformTLS::AllocTlsSlot();
-		}
-		check(FPlatformTLS::IsValidTlsSlot(BinnedTlsSlot));
+
 		FPerThreadFreeBlockLists::SetTLS();
 	}
 
