Removing some STL usage from Core
- CaptureExpressionsAndValues should use TStringBuilder instead of std::ostringstream to join FStrings.
- EPathPermissionPrefixResult ostream overload isn't necessary.

#Synced-CL 51569266

diff --git a/Engine/Source/Runtime/Core/Public/Misc/LowLevelTestAdapter.h b/Engine/Source/Runtime/Core/Public/Misc/LowLevelTestAdapter.h
--- a/Engine/Source/Runtime/Core/Public/Misc/LowLevelTestAdapter.h
+++ b/Engine/Source/Runtime/Core/Public/Misc/LowLevelTestAdapter.h
@@ -4,10 +4,13 @@
 
 #if !WITH_LOW_LEVEL_TESTS
 
-#include "Containers/StringConv.h"
 #include "Containers/UnrealString.h"
+#include "Misc/StringBuilder.h"
+
+#if UE_ENABLE_INCLUDE_ORDER_DEPRECATED_IN_5_8
 #include <sstream>
 #include <vector>
+#endif
 
 /**
 * @brief Captures expressions and their evaluated values.
@@ -19,10 +22,11 @@
 template<typename... ArgTypes>
 FString CaptureExpressionsAndValues(const FString& InExpressions, ArgTypes&&... InExpressionsValues)
 {
-	std::ostringstream Result;
 	auto Args = { InExpressionsValues... };
 	auto Iter = Args.begin();
 
+	TStringBuilder<64> Result;
+
 	FString RemainingExpressions = InExpressions;
 	FString Expression;
 	while (RemainingExpressions.Split(TEXT(","), &Expression, &RemainingExpressions))
@@ -33,18 +37,18 @@ FString CaptureExpressionsAndValues(const FString& InExpressions, ArgTypes&&...
 		}
 		else if (Iter != Args.begin())
 		{
-			Result << ", ";
+			Result << TEXT(", ");
 		}
-		Result << std::string(TCHAR_TO_UTF8(*Expression.TrimStartAndEnd())) << " = "  << *Iter++;
+		Result << Expression.TrimStartAndEnd() << TEXT(" = ")  << *Iter++;
 	}
 
 	if (Iter != Args.end() && !RemainingExpressions.TrimStartAndEnd().IsEmpty())
 	{
-		Result << ", " << std::string(TCHAR_TO_UTF8(*RemainingExpressions)) << " = " << *Iter;
+		Result << TEXT(", ") << RemainingExpressions << TEXT(" = ") << *Iter;
 	}
 
-	Result << std::endl;
-	return FString(Result.str().c_str());
+	Result << TEXT("\n");
+	return FString(Result.ToString());
 }
 
 
@@ -158,4 +162,4 @@ FString CaptureExpressionsAndValues(const FString& InExpressions, ArgTypes&&...
 
 #define SKIP(Message)
 
-#endif // !WITH_LOW_LEVEL_TESTS
\ No newline at end of file
+#endif // !WITH_LOW_LEVEL_TESTS
diff --git a/Engine/Source/Runtime/Core/Tests/Misc/PathPermissionListTests.cpp b/Engine/Source/Runtime/Core/Tests/Misc/PathPermissionListTests.cpp
--- a/Engine/Source/Runtime/Core/Tests/Misc/PathPermissionListTests.cpp
+++ b/Engine/Source/Runtime/Core/Tests/Misc/PathPermissionListTests.cpp
@@ -7,22 +7,6 @@
 
 #include "Tests/TestHarnessAdapter.h"
 
-std::ostream& operator<<(std::ostream& Stream, EPathPermissionPrefixResult Value)
-{
-    switch(Value)
-    {
-        case EPathPermissionPrefixResult::Fail:
-            return Stream << "Fail";
-        case EPathPermissionPrefixResult::FailRecursive:
-            return Stream << "FailRecursive";
-        case EPathPermissionPrefixResult::Pass:
-            return Stream << "Pass";
-        case EPathPermissionPrefixResult::PassRecursive:
-            return Stream << "PassRecursive";
-    }
-    return Stream << "Unknown";
-}
-
 TEST_CASE_NAMED(FPathPermissionListExactMatchTests, "System::Core::Misc::PathPermissionList::ExactMatch", "[Core][Misc][PathPermissionList]")
 {
     SECTION("Deny all")
