Added back weak reference bumping up and down in TSharedFromThis

#Synced-CL 51535387

diff --git a/Engine/Source/Runtime/Core/Public/Templates/SharedPointer.h b/Engine/Source/Runtime/Core/Public/Templates/SharedPointer.h
--- a/Engine/Source/Runtime/Core/Public/Templates/SharedPointer.h
+++ b/Engine/Source/Runtime/Core/Public/Templates/SharedPointer.h
@@ -1878,7 +1878,20 @@ private:
 	{
 		static_assert(OtherMode == Mode, "You cannot create a shared reference or pointer of one mode with a type which inherits TSharedFromThis of another mode.");
 
-		ReferenceController = InReferenceController;
+		if (!DoesSharedInstanceExist() && ReferenceController != InReferenceController)
+		{
+			if (InReferenceController)
+			{
+				InReferenceController->AddWeakReference();
+			}
+			
+			if (ReferenceController)
+			{
+				ReferenceController->ReleaseWeakReference();
+			}
+			
+			ReferenceController = InReferenceController;
+		}
 	}
 
 public:
@@ -1918,7 +1931,13 @@ protected:
 	}
 
 	/** Hidden destructor */
-	~TSharedFromThis() = default;
+	~TSharedFromThis()
+	{
+		if (ReferenceController)
+		{
+			ReferenceController->ReleaseWeakReference();
+		}
+	}
 
 private:
 	
