Support for taking trace snapshots on crash

Allow for crashes to contain a trace snapshot. Use cvar "CrashContext.TraceSnapshotOnCrash" to enable and "CrashContext.TraceSnapshotOnCrashTailSize" to set size of tail data to be written. Only supported on Windows at this point.

#Synced-CL 51335054

diff --git a/Engine/Source/Runtime/Core/Private/GenericPlatform/GenericPlatformCrashContext.cpp b/Engine/Source/Runtime/Core/Private/GenericPlatform/GenericPlatformCrashContext.cpp
--- a/Engine/Source/Runtime/Core/Private/GenericPlatform/GenericPlatformCrashContext.cpp
+++ b/Engine/Source/Runtime/Core/Private/GenericPlatform/GenericPlatformCrashContext.cpp
@@ -1,36 +1,38 @@
 // Copyright Epic Games, Inc. All Rights Reserved.
 
 #include "GenericPlatform/GenericPlatformCrashContext.h"
-#include "HAL/PlatformTime.h"
+#include "Containers/StringFwd.h"
+#include "Containers/Ticker.h"
+#include "HAL/FileManager.h"
+#include "HAL/IConsoleManager.h"
+#include "HAL/PlatformFileManager.h"
+#include "HAL/PlatformOutputDevices.h"
 #include "HAL/PlatformStackWalk.h"
-#include "Misc/Char.h"
-#include "Misc/Parse.h"
-#include "Misc/FileHelper.h"
-#include "Misc/CommandLine.h"
-#include "Misc/Paths.h"
+#include "HAL/PlatformTime.h"
 #include "Internationalization/Culture.h"
-#include "Misc/Optional.h"
 #include "Internationalization/Internationalization.h"
-#include "Misc/Guid.h"
-#include "Misc/SecureHash.h"
-#include "Containers/Ticker.h"
-#include "Containers/StringFwd.h"
+#include "Internationalization/Regex.h"
+#include "Internationalization/TextLocalizationManager.h"
+#include "Logging/LogScopedCategoryAndVerbosityOverride.h"
+#include "Misc/App.h"
+#include "Misc/Char.h"
+#include "Misc/CommandLine.h"
 #include "Misc/ConfigCacheIni.h"
 #include "Misc/CoreDelegates.h"
-#include "Misc/Fork.h"
-#include "Misc/App.h"
-#include "Misc/EngineVersion.h"
 #include "Misc/EngineBuildSettings.h"
+#include "Misc/EngineVersion.h"
+#include "Misc/FileHelper.h"
+#include "Misc/Fork.h"
+#include "Misc/Guid.h"
+#include "Misc/Optional.h"
+#include "Misc/OutputDeviceArchiveWrapper.h"
 #include "Misc/OutputDeviceRedirector.h"
+#include "Misc/Parse.h"
+#include "Misc/Paths.h"
 #include "Misc/ScopeLock.h"
+#include "Misc/SecureHash.h"
+#include "ProfilingDebugging/TraceAuxiliary.h"
 #include "Stats/Stats.h"
-#include "Internationalization/Regex.h"
-#include "Internationalization/TextLocalizationManager.h"
-#include "Logging/LogScopedCategoryAndVerbosityOverride.h"
-#include "HAL/FileManager.h"
-#include "HAL/PlatformOutputDevices.h"
-#include "HAL/PlatformFileManager.h"
-#include "Misc/OutputDeviceArchiveWrapper.h"
 
 #ifndef NOINITCRASHREPORTER
 #define NOINITCRASHREPORTER 0
@@ -49,6 +51,20 @@
 
 extern CORE_API bool GIsGPUCrashed;
 
+static TAutoConsoleVariable<bool> CVarTraceSnapshotOnCrash(
+	TEXT("CrashContext.TraceSnapshotOnCrash"),
+	false,
+	TEXT("Execute a trace snapshot and include in the crash report."),
+	ECVF_ReadOnly
+);
+
+static TAutoConsoleVariable<int32> CVarTraceSnapshotOnCrashTailSize(
+	TEXT("CrashContext.TraceSnapshotOnCrashTailSize"),
+	0,
+	TEXT("Max limit of tail size to write to snapshot. Zero indicates no limit."),
+	ECVF_ReadOnly
+);
+
 /**
  * A function-like type that creates a TStringBuilder to xml-escape a string
  */
@@ -1747,6 +1763,14 @@
 	}
 }
 
+void FGenericCrashContext::GenerateFiles(const TCHAR* OutputDirectory)
+{
+	if (CVarTraceSnapshotOnCrash.GetValueOnAnyThread())
+	{
+		FString SnapshotPath = FPaths::Combine(OutputDirectory, TEXT("Snapshot.utrace"));
+		FTraceAuxiliary::WriteSnapshot(*SnapshotPath, CVarTraceSnapshotOnCrashTailSize.GetValueOnAnyThread());
+	}
+}
 
 void FGenericCrashContext::CopyPlatformSpecificFiles(const TCHAR* OutputDirectory, void* Context)
 {
diff --git a/Engine/Source/Runtime/Core/Private/Windows/WindowsPlatformCrashContext.cpp b/Engine/Source/Runtime/Core/Private/Windows/WindowsPlatformCrashContext.cpp
--- a/Engine/Source/Runtime/Core/Private/Windows/WindowsPlatformCrashContext.cpp
+++ b/Engine/Source/Runtime/Core/Private/Windows/WindowsPlatformCrashContext.cpp
@@ -112,6 +112,7 @@ const uint32 OutOfMemoryExceptionCode = 0xc000;
 constexpr double DefaultCrashHandlingTimeoutSecs = 60.0;
 
 namespace {
+
 	/**
 	 * Write a Windows minidump to disk
 	 * @param The Crash context with its data already serialized into its buffer
@@ -897,6 +898,9 @@ int32 ReportCrashForMonitor(
 		// Copy the log file to output
 		FGenericCrashContext::DumpLog(CrashDirectoryAbsolute);
 
+		// Generate some files that needs to be done in the crashing process
+		FGenericCrashContext::GenerateFiles(*CrashDirectoryAbsolute);
+
 		// Dump additional context
 		GuardedDumpAdditionalContext(*CrashDirectoryAbsolute);
 	}
@@ -1025,7 +1029,10 @@ int32 ReportCrashUsingCrashReportClient(FWindowsPlatformCrashContext& InContext,
 			const FString CrashContextXMLPath = FPaths::Combine(*CrashFolderAbsolute, FPlatformCrashContext::CrashContextRuntimeXMLNameW);
 			InContext.SerializeAsXML(*CrashContextXMLPath);
 
-			// Copy platform specific files (e.g. minidump) to output directory
+			// Generate some files that needs to be done in the crashing process
+			FGenericCrashContext::GenerateFiles(*CrashFolderAbsolute);
+
+			// Generate and copy platform specific files (e.g. minidump) to output directory
 			InContext.CopyPlatformSpecificFiles(*CrashFolderAbsolute, (void*) ExceptionInfo);
 
 			// Dump additional context
diff --git a/Engine/Source/Runtime/Core/Public/GenericPlatform/GenericPlatformCrashContext.h b/Engine/Source/Runtime/Core/Public/GenericPlatform/GenericPlatformCrashContext.h
--- a/Engine/Source/Runtime/Core/Public/GenericPlatform/GenericPlatformCrashContext.h
+++ b/Engine/Source/Runtime/Core/Public/GenericPlatform/GenericPlatformCrashContext.h
@@ -732,6 +732,9 @@ public:
 	CORE_API virtual void AddPortableThreadCallStacks(TConstArrayView<FThreadCallStack> Threads);
 	CORE_API virtual void AddPortableThreadCallStack(uint32 ThreadId, const TCHAR* ThreadName, const uint64* StackFrames, int32 NumStackFrames);
 
+	/** Generates files required to be created in the crashing process */
+	CORE_API static void GenerateFiles(const TCHAR* OutputDirectory);
+
 	/** Allows platform implementations to copy files to report directory. */
 	CORE_API virtual void CopyPlatformSpecificFiles(const TCHAR* OutputDirectory, void* Context);
 
