Fixed default browser lookup failing in some Windows 11 environments.

diff --git a/Engine/Source/Runtime/Core/Private/Windows/WindowsPlatformProcess.cpp b/Engine/Source/Runtime/Core/Private/Windows/WindowsPlatformProcess.cpp
--- a/Engine/Source/Runtime/Core/Private/Windows/WindowsPlatformProcess.cpp
+++ b/Engine/Source/Runtime/Core/Private/Windows/WindowsPlatformProcess.cpp
@@ -37,6 +37,7 @@
 #include "Windows/AllowWindowsPlatformTypes.h"
 	#include <shellapi.h>
 	#include <ShlObj.h>
+	#include <Shlwapi.h>
 	#include <LM.h>
 	#include <Psapi.h>
 	#include <TlHelp32.h>
@@ -323,20 +324,36 @@ static void LaunchWebURL( const FString& URLParams, FString* Error )
 	UE_LOG(LogWindows, Log, TEXT("LaunchURL %s"), *URLParams);
 
 	FString BrowserOpenCommand;
+	BrowserOpenCommand.GetCharArray().SetNumUninitialized(MAX_PATH);
 
-	// First lookup the program Id for the default browser.
-	FString ProgId;
-	if (FWindowsPlatformMisc::QueryRegKey(HKEY_CURRENT_USER, TEXT("Software\\Microsoft\\Windows\\Shell\\Associations\\UrlAssociations\\http\\UserChoice"), TEXT("Progid"), ProgId))
+	// Look up the shell open command for the http protocol
+	bool bRetry = true;
+	for (;;)
 	{
-		// If we found it, then lookup it's open shell command in the classes registry.
-		FString BrowserRegPath = ProgId + TEXT("\\shell\\open\\command");
-		FWindowsPlatformMisc::QueryRegKey(HKEY_CLASSES_ROOT, *BrowserRegPath, NULL, BrowserOpenCommand);
-	}
+		// Query the command
+		DWORD Size = BrowserOpenCommand.GetCharArray().Num();
+		HRESULT hResult = AssocQueryString(ASSOCF_NOTRUNCATE | ASSOCF_IS_PROTOCOL, ASSOCSTR_COMMAND, TEXT("http"), TEXT("open"), BrowserOpenCommand.GetCharArray().GetData(), &Size);
 
-	// If we failed to find a default browser using the newer location, revert to using shell open command for the HTTP file association.
-	if (BrowserOpenCommand.IsEmpty())
-	{
-		FWindowsPlatformMisc::QueryRegKey(HKEY_CLASSES_ROOT, TEXT("http\\shell\\open\\command"), NULL, BrowserOpenCommand);
+		// Retry once if the buffer isn't big enough
+		if (hResult == E_POINTER || hResult == S_FALSE || hResult == HRESULT_FROM_WIN32(ERROR_MORE_DATA))
+		{
+			if (bRetry && Size <= (DWORD)MIN_int32)
+			{
+				bRetry = false;
+				BrowserOpenCommand.GetCharArray().SetNum((int32)Size);
+				continue;
+			}
+		}
+
+		// If there's still an error, assume the command couldn't be found
+		if (hResult != S_OK)
+		{
+			Size = 0;
+		}
+
+		// Resize the string to the right size
+		BrowserOpenCommand.GetCharArray().SetNum((int32)Size);
+		break;
 	}
 
 	// If we have successfully looked up the correct shell command, then we can create a new process using that command
