[Core] Apply AUTORTFM_INFER to templated containers / wrapper utilities.

Without this, we can end up with AutoRTFM errors where these (previously AutoRTFM-enabled) functions call AutoRTFM-disabled functions.

Fixes build failures in LTO / PGO builds.

#Synced-CL 51584942

diff --git a/Engine/Source/Runtime/Core/Public/Containers/Array.h b/Engine/Source/Runtime/Core/Public/Containers/Array.h
--- a/Engine/Source/Runtime/Core/Public/Containers/Array.h
+++ b/Engine/Source/Runtime/Core/Public/Containers/Array.h
@@ -2,6 +2,7 @@
 
 #pragma once
 
+#include "AutoRTFM/Defines.h"
 #include "AutoRTFM/ReturnFromOpenMode.h"
 #include "CoreTypes.h"
 #include "Misc/AssertionMacros.h"
@@ -764,7 +765,7 @@ namespace UE::Core::Private
  *
  **/
 template<typename InElementType, typename InAllocatorType>
-class TArray
+class AUTORTFM_INFER TArray
 {
 	template <typename OtherInElementType, typename OtherAllocator>
 	friend class TArray;
diff --git a/Engine/Source/Runtime/Core/Public/Containers/Map.h b/Engine/Source/Runtime/Core/Public/Containers/Map.h
--- a/Engine/Source/Runtime/Core/Public/Containers/Map.h
+++ b/Engine/Source/Runtime/Core/Public/Containers/Map.h
@@ -4,6 +4,7 @@
 
 #include "CoreTypes.h"
 
+#include "AutoRTFM/Defines.h"
 #include "Algo/Reverse.h"
 #include "Containers/ContainerElementTypeCompatibility.h"
 #include "Containers/Set.h"
@@ -55,7 +56,7 @@ public:
 
 /** An initializer type for keys that's passed to the pair set when adding a new key. */
 template <typename KeyInitType>
-class TKeyInitializer
+class AUTORTFM_INFER TKeyInitializer
 {
 public:
 	std::conditional_t<std::is_rvalue_reference_v<KeyInitType>, KeyInitType&, KeyInitType> Key;
@@ -75,7 +76,7 @@ public:
 
 /** Defines how the map's pairs are hashed. */
 template<typename KeyType, typename ValueType, bool bInAllowDuplicateKeys>
-struct TDefaultMapKeyFuncs : BaseKeyFuncs<TPair<KeyType, ValueType>, KeyType, bInAllowDuplicateKeys>
+struct AUTORTFM_INFER TDefaultMapKeyFuncs : BaseKeyFuncs<TPair<KeyType, ValueType>, KeyType, bInAllowDuplicateKeys>
 {
 	typedef typename TTypeTraits<KeyType>::ConstPointerType KeyInitType;
 	typedef const TPairInitializer<typename TTypeTraits<KeyType>::ConstInitType, typename TTypeTraits<ValueType>::ConstInitType>& ElementInitType;
diff --git a/Engine/Source/Runtime/Core/Public/Containers/Map.h.inl b/Engine/Source/Runtime/Core/Public/Containers/Map.h.inl
--- a/Engine/Source/Runtime/Core/Public/Containers/Map.h.inl
+++ b/Engine/Source/Runtime/Core/Public/Containers/Map.h.inl
@@ -25,7 +25,7 @@
  *    the cache misses of reading key data and doing the hashing *before* acquiring the lock.
  **/
 template <typename KeyType, typename ValueType, typename SetAllocator, typename KeyFuncs>
-class TMAPBASE
+class AUTORTFM_INFER TMAPBASE
 {
 	template <typename OtherKeyType, typename OtherValueType, typename OtherSetAllocator, typename OtherKeyFuncs>
 	friend class TMAPBASE;
@@ -1018,7 +1018,7 @@ public:
 
 /** The base type of sortable maps. */
 template <typename KeyType, typename ValueType, typename SetAllocator, typename KeyFuncs>
-class TSORTABLEMAPBASE : public TMAPBASE<KeyType, ValueType, SetAllocator, KeyFuncs>
+class AUTORTFM_INFER TSORTABLEMAPBASE : public TMAPBASE<KeyType, ValueType, SetAllocator, KeyFuncs>
 {
 protected:
 	typedef TMAPBASE<KeyType, ValueType, SetAllocator, KeyFuncs> Super;
@@ -1175,7 +1175,7 @@ private:
 
 /** A TMAPBASE specialization that only allows a single value associated with each key.*/
 template<typename InKeyType, typename InValueType, typename SetAllocator /*= FDefaultSetAllocator*/, typename KeyFuncs /*= TDefaultMapHashableKeyFuncs<KeyType,ValueType,false>*/>
-class TMAP : public TSORTABLEMAPBASE<InKeyType, InValueType, SetAllocator, KeyFuncs>
+class AUTORTFM_INFER TMAP : public TSORTABLEMAPBASE<InKeyType, InValueType, SetAllocator, KeyFuncs>
 {
 	template <typename, typename>
 	friend class TScriptMap;
diff --git a/Engine/Source/Runtime/Core/Public/Containers/SetUtilities.h b/Engine/Source/Runtime/Core/Public/Containers/SetUtilities.h
--- a/Engine/Source/Runtime/Core/Public/Containers/SetUtilities.h
+++ b/Engine/Source/Runtime/Core/Public/Containers/SetUtilities.h
@@ -2,6 +2,7 @@
 
 #pragma once
 
+#include "AutoRTFM/Defines.h"
 #include "CoreTypes.h"
 #include "Serialization/MemoryLayout.h"
 #include "Templates/MemoryOps.h"
@@ -37,7 +38,7 @@ struct BaseKeyFuncs
  * A default implementation of the KeyFuncs used by TSet which uses the element as a key.
  */
 template<typename ElementType,bool bInAllowDuplicateKeys /*= false*/>
-struct DefaultKeyFuncs : BaseKeyFuncs<ElementType,ElementType,bInAllowDuplicateKeys>
+struct AUTORTFM_INFER DefaultKeyFuncs : BaseKeyFuncs<ElementType,ElementType,bInAllowDuplicateKeys>
 {
 	typedef typename TTypeTraits<ElementType>::ConstPointerType KeyInitType;
 	typedef typename TCallTraits<ElementType>::ParamType ElementInitType;
@@ -84,7 +85,7 @@ struct DefaultKeyFuncs : BaseKeyFuncs<ElementType,ElementType,bInAllowDuplicateK
 /** This is used to provide type specific behavior for a move which will destroy B. */
 /** Should be in UnrealTemplate but isn't for Clang build reasons - will move later */
 template<typename T>
-inline void MoveByRelocate(T& A, T& B)
+AUTORTFM_INFER inline void MoveByRelocate(T& A, T& B)
 {
 	// Destruct the previous value of A.
 	A.~T();
@@ -140,4 +141,4 @@ private:
 };
 
 // This is just an int32
-DECLARE_INTRINSIC_TYPE_LAYOUT(FSetElementId);
\ No newline at end of file
+DECLARE_INTRINSIC_TYPE_LAYOUT(FSetElementId);
diff --git a/Engine/Source/Runtime/Core/Public/Containers/SparseArray.h b/Engine/Source/Runtime/Core/Public/Containers/SparseArray.h
--- a/Engine/Source/Runtime/Core/Public/Containers/SparseArray.h
+++ b/Engine/Source/Runtime/Core/Public/Containers/SparseArray.h
@@ -2,6 +2,7 @@
 
 #pragma once
 
+#include "AutoRTFM/Defines.h"
 #include "CoreTypes.h"
 #include "Misc/AssertionMacros.h"
 #include "HAL/UnrealMemory.h"
@@ -520,7 +521,7 @@ class TScriptSparseArray;
  *
  **/
 template<typename InElementType, typename Allocator /*= FDefaultSparseArrayAllocator */>
-class TSparseArray : public TSparseArrayBase<sizeof(InElementType), alignof(InElementType), Allocator>
+class AUTORTFM_INFER TSparseArray : public TSparseArrayBase<sizeof(InElementType), alignof(InElementType), Allocator>
 {
 	using SuperType = TSparseArrayBase<sizeof(InElementType), alignof(InElementType), Allocator>;
 	using ElementType = InElementType;
diff --git a/Engine/Source/Runtime/Core/Public/Containers/SparseSet.h.inl b/Engine/Source/Runtime/Core/Public/Containers/SparseSet.h.inl
--- a/Engine/Source/Runtime/Core/Public/Containers/SparseSet.h.inl
+++ b/Engine/Source/Runtime/Core/Public/Containers/SparseSet.h.inl
@@ -4,6 +4,7 @@
 #error "SparseSet.h.inl should only be included after defining UE_TSPARSE_SET"
 #endif
 
+#include "AutoRTFM/Defines.h"
 #include "Containers/ContainerAllocationPolicies.h"
 #include "Containers/ContainerElementTypeCompatibility.h"
 #include "Containers/SetUtilities.h"
@@ -55,7 +56,7 @@ template<
 	typename KeyFuncs /*= DefaultKeyFuncs<ElementType>*/,
 	typename Allocator /*= FDefaultSparseSetAllocator*/
 	>
-class UE_TSPARSE_SET
+class AUTORTFM_INFER UE_TSPARSE_SET
 {
 public:
 	typedef InElementType ElementType;
@@ -1854,4 +1855,4 @@ template <typename ElementType, typename KeyFuncs, typename Allocator> struct TI
 template <typename ElementType, typename KeyFuncs, typename Allocator> struct TIsTSet<      volatile UE_TSPARSE_SET<ElementType, KeyFuncs, Allocator>> { enum { Value = true }; };
 template <typename ElementType, typename KeyFuncs, typename Allocator> struct TIsTSet<const volatile UE_TSPARSE_SET<ElementType, KeyFuncs, Allocator>> { enum { Value = true }; };
 
-#undef TSETPRIVATEFRIEND
\ No newline at end of file
+#undef TSETPRIVATEFRIEND
diff --git a/Engine/Source/Runtime/Core/Public/Containers/SparseSetElement.h b/Engine/Source/Runtime/Core/Public/Containers/SparseSetElement.h
--- a/Engine/Source/Runtime/Core/Public/Containers/SparseSetElement.h
+++ b/Engine/Source/Runtime/Core/Public/Containers/SparseSetElement.h
@@ -2,13 +2,14 @@
 
 #pragma once
 
+#include "AutoRTFM/Defines.h"
 #include "Containers/SetUtilities.h"
 #include "Serialization/StructuredArchive.h"
 #include <initializer_list>
 #include <type_traits>
 
 template<typename InElementType, bool bTypeLayout>
-class TSparseSetElementBase
+class AUTORTFM_INFER TSparseSetElementBase
 {
 public:
 	typedef InElementType ElementType;
@@ -41,7 +42,7 @@ public:
 };
 
 template<typename InElementType>
-class TSparseSetElementBase<InElementType, true>
+class AUTORTFM_INFER TSparseSetElementBase<InElementType, true>
 {
 	DECLARE_INLINE_TYPE_LAYOUT(TSparseSetElementBase, NonVirtual);
 public:
@@ -76,7 +77,7 @@ public:
 
 /** An element in the set. */
 template <typename InElementType>
-class TSparseSetElement : public TSparseSetElementBase<InElementType, THasTypeLayout<InElementType>::Value>
+class AUTORTFM_INFER TSparseSetElement : public TSparseSetElementBase<InElementType, THasTypeLayout<InElementType>::Value>
 {
 	using Super = TSparseSetElementBase<InElementType, THasTypeLayout<InElementType>::Value>;
 public:
diff --git a/Engine/Source/Runtime/Core/Public/Misc/Optional.h b/Engine/Source/Runtime/Core/Public/Misc/Optional.h
--- a/Engine/Source/Runtime/Core/Public/Misc/Optional.h
+++ b/Engine/Source/Runtime/Core/Public/Misc/Optional.h
@@ -2,6 +2,7 @@
 
 #pragma once
 
+#include "AutoRTFM/Defines.h"
 #include "AutoRTFM/ReturnFromOpenMode.h"
 #include "CoreTypes.h"
 #include "Misc/AssertionMacros.h"
@@ -44,7 +45,7 @@ namespace UE::Core::Private
  * Otherwise GetValue() is not meaningful.
  */
 template<typename OptionalType>
-struct TOptional
+struct AUTORTFM_INFER TOptional
 {
 private:
 	using FOptional = UE::Core::Private::FOptional;
diff --git a/Engine/Source/Runtime/Core/Public/Templates/Tuple.h b/Engine/Source/Runtime/Core/Public/Templates/Tuple.h
--- a/Engine/Source/Runtime/Core/Public/Templates/Tuple.h
+++ b/Engine/Source/Runtime/Core/Public/Templates/Tuple.h
@@ -2,7 +2,7 @@
 
 #pragma once
 
-#include "AutoRTFM.h"
+#include "AutoRTFM/Defines.h"
 #include "CoreTypes.h"
 #include "Concepts/NotCVRefTo.h"
 #include "Templates/UnrealTemplate.h"
@@ -57,7 +57,7 @@ namespace UE::Core::Private::Tuple
 	constexpr uint32 TTypeCountInParameterPack_V<T, U, Types...> = TTypeCountInParameterPack_V<T, Types...> + (std::is_same_v<T, U> ? 1 : 0);
 
 	template <typename T, uint32 Index, uint32 TupleSize>
-	struct TTupleBaseElement
+	struct AUTORTFM_INFER TTupleBaseElement
 	{
 		template <typename ArgType>
 		constexpr explicit TTupleBaseElement(EForwardingConstructor, ArgType&& Arg)
diff --git a/Engine/Source/Runtime/Core/Public/Templates/UniquePtr.h b/Engine/Source/Runtime/Core/Public/Templates/UniquePtr.h
--- a/Engine/Source/Runtime/Core/Public/Templates/UniquePtr.h
+++ b/Engine/Source/Runtime/Core/Public/Templates/UniquePtr.h
@@ -2,6 +2,7 @@
 
 #pragma once
 
+#include "AutoRTFM/Defines.h"
 #include "CoreTypes.h"
 #include "Templates/UnrealTemplate.h"
 #include "Templates/IsArray.h"
@@ -24,7 +25,7 @@
 // https://en.cppreference.com/w/cpp/memory/unique_ptr/~unique_ptr
 
 template <typename T>
-struct TDefaultDelete
+struct AUTORTFM_INFER TDefaultDelete
 {
 	DECLARE_INLINE_TYPE_LAYOUT(TDefaultDelete, NonVirtual);
 
@@ -103,7 +104,7 @@ struct TDefaultDelete<T[]>
 };
 
 template <typename T, typename Deleter = TDefaultDelete<T>>
-class TUniquePtr : private Deleter
+class AUTORTFM_INFER TUniquePtr : private Deleter
 {
 	DECLARE_INLINE_TYPE_LAYOUT_EXPLICIT_BASES(TUniquePtr, NonVirtual, Deleter);
 
@@ -496,7 +497,7 @@ public:
 };
 
 template <typename T, typename Deleter>
-class TUniquePtr<T[], Deleter> : private Deleter
+class AUTORTFM_INFER TUniquePtr<T[], Deleter> : private Deleter
 {
 	template <typename OtherT, typename OtherDeleter>
 	friend class TUniquePtr;
