ParkingLot: Use low-level virtual memory API for allocations to allow our allocators to depend on ParkingLot

Extended TSizedInlineAllocator with SupportsElementAlignment support.

#Synced-CL 50967057

diff --git a/Engine/Source/Runtime/Core/Private/Async/ParkingLot.cpp b/Engine/Source/Runtime/Core/Private/Async/ParkingLot.cpp
--- a/Engine/Source/Runtime/Core/Private/Async/ParkingLot.cpp
+++ b/Engine/Source/Runtime/Core/Private/Async/ParkingLot.cpp
@@ -9,9 +9,16 @@
 #include "Containers/Array.h"
 #include "Containers/ArrayView.h"
 #include "HAL/PlatformManualResetEvent.h"
+#include "Memory/VirtualLinearAllocator.h"
 #include "Templates/RefCounting.h"
+#include "Templates/NoDestroy.h"
 #include <atomic>
 
+#if !UE_ENABLE_VIRTUAL_LINEAR_ALLOCATOR
+#include "HAL/Allocators/AnsiAllocator.h"
+#include "HAL/MallocAnsi.h"
+#endif
+
 namespace UE::ParkingLot::Private
 {
 
@@ -44,6 +51,188 @@ static thread_local EExecutingFunction ParkingLotExecutingFunction = EExecutingF
 
 ///////////////////////////////////////////////////////////////////////////////////////////////////////////////////////
 
+#if UE_ENABLE_VIRTUAL_LINEAR_ALLOCATOR
+/**
+ * Allocator for permanent allocations by ParkingLot: FTable, FBucket, FThread.
+ *
+ * Shared by FTable, TObjectAllocator<FBucket>, TObjectAllocator<FThread>.
+ * ParkingLot allocates approximately ThreadCount buckets and threads, at 64 bytes each,
+ * and tables totaling approximately 2*ThreadCount 8-byte pointers.
+ * Even 16,384 threads will consume only 2.25 MiB, which makes 4 MiB plenty.
+ *
+ * Adds 64 KiB to avoid a multiple of 2 MiB as it can cause issues on platforms with transparent large pages.
+ */
+static TNoDestroy<TVirtualLinearAllocator<FWordMutex>> ParkingLotPermanentAllocator(4 * 1024 * 1024 + 64 * 1024);
+#else
+static struct
+{
+	UE_REWRITE void* Allocate(SIZE_T Size, uint32 Alignment)
+	{
+		return AnsiMalloc(Size, Alignment);
+	}
+} ParkingLotPermanentAllocator;
+#endif
+
+///////////////////////////////////////////////////////////////////////////////////////////////////////////////////////
+
+/**
+ * A low-level allocator that maintains a free list for a given object type. Memory is never released.
+ */
+template <typename ObjectType>
+class TObjectAllocator
+{
+	static_assert(sizeof(ObjectType) >= sizeof(void*));
+	static_assert(alignof(ObjectType) >= alignof(void*));
+
+public:
+	static void* Allocate()
+	{
+		if (TUniqueLock Lock(FreeListMutex); void* Mem = FreeListHead)
+		{
+			FreeListHead = *(void**)Mem;
+			return Mem;
+		}
+		return ParkingLotPermanentAllocator->Allocate(sizeof(ObjectType), alignof(ObjectType));
+	}
+
+	static void Free(void* Mem)
+	{
+		TUniqueLock Lock(FreeListMutex);
+		void* Next = FreeListHead;
+		*(void**)Mem = Next;
+		FreeListHead = Mem;
+	}
+
+private:
+	inline static FWordMutex FreeListMutex;
+	inline static void* FreeListHead = nullptr;
+};
+
+///////////////////////////////////////////////////////////////////////////////////////////////////////////////////////
+
+#if UE_ENABLE_VIRTUAL_LINEAR_ALLOCATOR
+/**
+ * Allocator for scratch allocations by ParkingLot: TArray<PointerType>.
+ *
+ * Scratch allocations are bounded by ThreadCount:
+ * 1. FTable::Reserve can allocate maximum 3*ThreadCount 8-byte pointers while locked,
+ *    and ThreadCount 8-byte pointers while unlocked. Reserve is only called while starting
+ *    a new thread, which limits the number of concurrent calls. Even 256 new threads racing
+ *    to increase the size to 1024 will allocate a worst-case ~2 MiB.
+ * 2. WakeMultiple can allocate maximum WaitingThreadCount 8-byte pointers across every
+ *    thread that is waking other threads. Wakes can race with timeouts to bring the
+ *    effective total higher but even 512 threads waking 512 threads is only 2 MiB.
+ *
+ * Adds 64 KiB to avoid a multiple of 2 MiB as it can cause issues on platforms with transparent large pages.
+ */
+inline static TNoDestroy<TVirtualLinearAllocator<FWordMutex>> ParkingLotScratchAllocator(4 * 1024 * 1024 + 64 * 1024);
+
+/**
+ * A container allocator backed by a virtual linear allocator for scratch allocations.
+ */
+class FScratchAllocator
+{
+public:
+	using SizeType = int32;
+	enum { NeedsElementType = false };
+	enum { RequireRangeCheck = false };
+	enum { ShrinkByDefault = false };
+
+	class ForAnyElementType
+	{
+	public:
+		constexpr ForAnyElementType() = default;
+
+		inline void MoveToEmpty(ForAnyElementType& Other)
+		{
+			Data = Other.Data;
+			Other.Data = nullptr;
+		}
+
+		inline void ResizeAllocation(SizeType CurrentNum, SizeType NewMax, SIZE_T NumBytesPerElement, uint32 AlignmentOfElement)
+		{
+			if (NewMax)
+			{
+				FScriptContainerElement* const OldData = Data;
+				Data = (FScriptContainerElement*)ParkingLotScratchAllocator->Allocate(NewMax * NumBytesPerElement, AlignmentOfElement);
+				if (OldData && CurrentNum > 0)
+				{
+					FMemory::Memcpy(Data, OldData, CurrentNum * NumBytesPerElement);
+				}
+			}
+			else
+			{
+				Data = nullptr;
+			}
+		}
+
+		inline SizeType CalculateSlackReserve(SizeType NewMax, SIZE_T NumBytesPerElement, uint32 AlignmentOfElement) const
+		{
+			return DefaultCalculateSlackReserve(NewMax, NumBytesPerElement, /*bAllowQuantize*/ false, AlignmentOfElement);
+		}
+
+		inline SizeType CalculateSlackShrink(SizeType NewMax, SizeType CurrentMax, SIZE_T NumBytesPerElement, uint32 AlignmentOfElement) const
+		{
+			return DefaultCalculateSlackShrink(NewMax, CurrentMax, NumBytesPerElement, /*bAllowQuantize*/ false, AlignmentOfElement);
+		}
+
+		inline SizeType CalculateSlackGrow(SizeType NewMax, SizeType CurrentMax, SIZE_T NumBytesPerElement, uint32 AlignmentOfElement) const
+		{
+			return DefaultCalculateSlackGrow(NewMax, CurrentMax, NumBytesPerElement, /*bAllowQuantize*/ false, AlignmentOfElement);
+		}
+
+		inline SIZE_T GetAllocatedSize(SizeType CurrentMax, SIZE_T NumBytesPerElement) const
+		{
+			return CurrentMax * NumBytesPerElement;
+		}
+
+		inline FScriptContainerElement* GetAllocation() const { return Data; }
+		inline bool HasAllocation() const { return !!Data; }
+
+		constexpr SizeType GetInitialCapacity() const { return 0; }
+
+	private:
+		ForAnyElementType(const ForAnyElementType&);
+		ForAnyElementType& operator=(const ForAnyElementType&);
+
+		FScriptContainerElement* Data = nullptr;
+	};
+
+	template <typename ElementType>
+	class ForElementType : public ForAnyElementType
+	{
+	public:
+		inline ElementType* GetAllocation() const
+		{
+			return (ElementType*)ForAnyElementType::GetAllocation();
+		}
+	};
+};
+#else
+static struct
+{
+	constexpr void AddRef() {}
+	constexpr void Release() {}
+} ParkingLotScratchAllocator;
+using FScratchAllocator = FAnsiAllocator;
+#endif // UE_ENABLE_VIRTUAL_LINEAR_ALLOCATOR
+
+} // UE::ParkingLot::Private
+
+#if UE_ENABLE_VIRTUAL_LINEAR_ALLOCATOR
+template <>
+struct TAllocatorTraits<UE::ParkingLot::Private::FScratchAllocator> : TAllocatorTraitsBase<UE::ParkingLot::Private::FScratchAllocator>
+{
+	enum { IsZeroConstruct = true };
+	enum { SupportsElementAlignment = true };
+};
+#endif
+
+///////////////////////////////////////////////////////////////////////////////////////////////////////////////////////
+
+namespace UE::ParkingLot::Private
+{
+
 /**
  * A thread as stored in the wait queue.
  */
@@ -64,17 +253,28 @@ struct alignas(PLATFORM_CACHE_LINE_SIZE) FThread final
 	{
 		if (ReferenceCount.fetch_sub(1, std::memory_order_acq_rel) == 1)
 		{
-			delete this;
+			this->~FThread();
+			Allocator.Free(this);
 		}
 	}
 
+	inline static FThread* New()
+	{
+		return new(Allocator.Allocate()) FThread;
+	}
+
+private:
 	FThread() = default;
 	~FThread() = default;
 
 	FThread(const FThread&) = delete;
 	FThread& operator=(const FThread&) = delete;
+
+	static TObjectAllocator<FThread> Allocator;
 };
 
+TObjectAllocator<FThread> FThread::Allocator;
+
 class FThreadLocalData
 {
 public:
@@ -88,9 +288,9 @@ public:
 		}
 		if (LocalThreadLocalData.bDestroyed)
 		{
-			return new FThread;
+			return FThread::New();
 		}
-		LocalThreadLocalData.Thread = new FThread;
+		LocalThreadLocalData.Thread = FThread::New();
 		return LocalThreadLocalData.Thread;
 	}
 
@@ -126,6 +326,17 @@ enum class EQueueAction
 class alignas(PLATFORM_CACHE_LINE_SIZE) FBucket final
 {
 public:
+	static FBucket* Create()
+	{
+		return new(Allocator.Allocate()) FBucket;
+	}
+
+	static void Destroy(FBucket* Bucket)
+	{
+		Bucket->~FBucket();
+		Allocator.Free(Bucket);
+	}
+
 	[[nodiscard]] inline TDynamicUniqueLock<FWordMutex> LockDynamic() { return TDynamicUniqueLock(Mutex); }
 
 	inline void Lock() { Mutex.Lock(); }
@@ -147,18 +358,22 @@ public:
 	template <typename VisitorType>
 	void DequeueIf(VisitorType&& Visitor);
 
+private:
 	FBucket() = default;
 	~FBucket() = default;
 
-private:
 	FBucket(const FBucket&) = delete;
 	FBucket& operator=(const FBucket&) = delete;
 
 	FWordMutex Mutex;
 	FThread* Head = nullptr;
 	FThread* Tail = nullptr;
+
+	static TObjectAllocator<FBucket> Allocator;
 };
 
+TObjectAllocator<FBucket> FBucket::Allocator;
+
 void FBucket::Enqueue(FThread* Thread)
 {
 	checkSlow(Thread);
@@ -261,7 +476,7 @@ private:
 	static void Destroy(FTable& Table);
 
 	/** Try to lock the whole table by locking every one of its buckets. */
-	static bool TryLock(FTable& Table, TArray<FBucket*>& OutBuckets);
+	static bool TryLock(FTable& Table, TArray<FBucket*, FScratchAllocator>& OutBuckets);
 	/** Unlock an array of buckets that was filled by TryLock. */
 	static void Unlock(TConstArrayView<FBucket*> LockedBuckets);
 
@@ -270,6 +485,7 @@ private:
 
 	/** Pointer to the current global table. Previous global tables are leaked. */
 	inline static std::atomic<FTable*> GlobalTable;
+	inline static FWordMutex CreateMutex;
 
 	FTable() = default;
 	~FTable() = default;
@@ -293,7 +509,7 @@ FBucket& FTable::FindOrCreateBucket(const void* Address, TDynamicUniqueLock<FWor
 	{
 		FTable& Table = CreateOrGet();
 		const uint32 Index = Hash % Table.BucketCount;
-		FBucket& Bucket = Table.FindOrCreateBucket(Index, [] { return new FBucket; });
+		FBucket& Bucket = Table.FindOrCreateBucket(Index, [] { return FBucket::Create(); });
 		OutLock = Bucket.LockDynamic();
 
 		if (LIKELY(&Table == GlobalTable.load(std::memory_order_acquire)))
@@ -347,7 +563,7 @@ FBucket& FTable::FindOrCreateBucket(uint32 Index, AllocatorType&& BucketAllocato
 		}
 		else
 		{
-			delete NewBucket;
+			FBucket::Destroy(NewBucket);
 		}
 		checkSlow(Bucket);
 	}
@@ -356,30 +572,27 @@ FBucket& FTable::FindOrCreateBucket(uint32 Index, AllocatorType&& BucketAllocato
 
 FTable& FTable::CreateOrGet()
 {
-	FTable* Table = GlobalTable.load(std::memory_order_acquire);
-
-	if (LIKELY(Table))
+	if (FTable* Table = GlobalTable.load(std::memory_order_acquire); LIKELY(Table))
 	{
 		return *Table;
 	}
 
-	FTable& NewTable = Create(MinSize);
-
-	if (LIKELY(GlobalTable.compare_exchange_strong(Table, &NewTable, std::memory_order_release, std::memory_order_acquire)))
+	TUniqueLock Lock(CreateMutex);
+	if (FTable* Table = GlobalTable.load(std::memory_order_acquire))
 	{
-		return NewTable;
+		return *Table;
 	}
 
-	Destroy(NewTable);
-
-	checkSlow(Table);
-	return *Table;
+	FTable& NewTable = Create(MinSize);
+	FTable* CompareTable = GlobalTable.exchange(&NewTable, std::memory_order_release);
+	checkSlow(!CompareTable);
+	return NewTable;
 }
 
 FTable& FTable::Create(const uint32 Size)
 {
-	const SIZE_T MemorySize = sizeof(FTable) + sizeof(FBucket*) * SIZE_T(Size);
-	void* const Memory = FMemory::Malloc(MemorySize, alignof(FTable));
+	const SIZE_T MemorySize = sizeof(FTable) + sizeof(std::atomic<FBucket*>) * SIZE_T(Size);
+	void* const Memory = ParkingLotPermanentAllocator->Allocate(MemorySize, alignof(FTable));
 	FMemory::Memzero(Memory, MemorySize);
 	FTable& Table = *new(Memory) FTable;
 	Table.BucketCount = Size;
@@ -389,13 +602,15 @@ FTable& FTable::Create(const uint32 Size)
 void FTable::Destroy(FTable& Table)
 {
 	Table.~FTable();
-	FMemory::Free(&Table);
 }
 
 void FTable::Reserve(const uint32 ThreadCount)
 {
+	// Reference the allocator to allow it to decommit scratch memory.
+	TRefCountPtr ScratchAllocatorRef(&ParkingLotScratchAllocator);
+
 	const uint32 TargetBucketCount = FMath::RoundUpToPowerOfTwo(ThreadCount);
-	TArray<FBucket*> ExistingBuckets;
+	TArray<FBucket*, FScratchAllocator> ExistingBuckets;
 
 	for (;;)
 	{
@@ -415,7 +630,7 @@ void FTable::Reserve(const uint32 ThreadCount)
 
 		// Gather waiting threads to be redistributed into the buckets of the new table.
 		// Threads with the same address remain in the same relative order as they were queued.
-		TArray<FThread*> Threads;
+		TArray<FThread*, TInlineAllocator<128, FScratchAllocator>> Threads;
 		for (FBucket* Bucket : ExistingBuckets)
 		{
 			while (FThread* Thread = Bucket->Dequeue())
@@ -427,10 +642,10 @@ void FTable::Reserve(const uint32 ThreadCount)
 		FTable& NewTable = Create(TargetBucketCount);
 
 		// Reuse existing now-empty buckets when populating the new table.
-		TArray<FBucket*> AvailableBuckets = ExistingBuckets;
+		TArray<FBucket*, FScratchAllocator> AvailableBuckets = ExistingBuckets;
 		const auto AllocateBucket = [&AvailableBuckets]() -> FBucket*
 		{
-			return !AvailableBuckets.IsEmpty() ? AvailableBuckets.Pop(EAllowShrinking::No) : new FBucket;
+			return !AvailableBuckets.IsEmpty() ? AvailableBuckets.Pop(EAllowShrinking::No) : FBucket::Create();
 		};
 
 		// Add waiting threads to the new table.
@@ -459,14 +674,14 @@ void FTable::Reserve(const uint32 ThreadCount)
 	}
 }
 
-bool FTable::TryLock(FTable& Table, TArray<FBucket*>& OutBuckets)
+bool FTable::TryLock(FTable& Table, TArray<FBucket*, FScratchAllocator>& OutBuckets)
 {
 	OutBuckets.Reset(Table.BucketCount);
 
 	// Gather buckets from the table, creating them as needed because the lock is on the bucket.
 	for (uint32 Index = 0; Index < Table.BucketCount; ++Index)
 	{
-		OutBuckets.Add(&Table.FindOrCreateBucket(Index, [] { return new FBucket; }));
+		OutBuckets.Add(&Table.FindOrCreateBucket(Index, [] { return FBucket::Create(); }));
 	}
 
 	// Lock the buckets in order by address to ensure consistent ordering regardless of the table being locked.
@@ -756,7 +971,9 @@ uint32 WakeMultiple(const void* const Address, const uint32 WakeCount)
 	checkf(ParkingLotExecutingFunction <= EExecutingFunction::BeforeWait,
 		TEXT("Detected a re-entrant call into ParkingLot from within a %s callback."), LexToString(ParkingLotExecutingFunction));
 
-	TArray<TRefCountPtr<FThread>, TInlineAllocator<128>> WakeThreads;
+	// Reference the allocator to allow it to decommit scratch memory.
+	TRefCountPtr ScratchAllocatorRef(&ParkingLotScratchAllocator);
+	TArray<TRefCountPtr<FThread>, TInlineAllocator<128, FScratchAllocator>> WakeThreads;
 
 	if (TDynamicUniqueLock<FWordMutex> BucketLock; FBucket* Bucket = FTable::FindBucket(Address, BucketLock))
 	{
diff --git a/Engine/Source/Runtime/Core/Public/Containers/ContainerAllocationPolicies.h b/Engine/Source/Runtime/Core/Public/Containers/ContainerAllocationPolicies.h
--- a/Engine/Source/Runtime/Core/Public/Containers/ContainerAllocationPolicies.h
+++ b/Engine/Source/Runtime/Core/Public/Containers/ContainerAllocationPolicies.h
@@ -954,29 +954,29 @@ public:
 			return GetInlineElements();
 		}
 
-		void ResizeAllocation(SizeType CurrentNum, SizeType NewMax,SIZE_T NumBytesPerElement)
+		void ResizeAllocation(SizeType CurrentNum, SizeType NewMax, SIZE_T NumBytesPerElement)
 		{
 			// Make sure the number of live elements is still withing the allocation since we only memmove, not destruct
 			checkSlow(CurrentNum <= NewMax);
 
 			// Check if the new allocation will fit in the inline data area.
-			if(NewMax <= NumInlineElements)
+			if (NewMax <= NumInlineElements)
 			{
 				// If the old allocation wasn't in the inline data area, relocate it into the inline data area.
-				if(SecondaryData.GetAllocation())
+				if (SecondaryData.GetAllocation())
 				{
 					RelocateConstructItems<ElementType>((void*)InlineData, (ElementType*)SecondaryData.GetAllocation(), CurrentNum);
 
 					// Free the old indirect allocation.
-					SecondaryData.ResizeAllocation(0,0,NumBytesPerElement);
+					SecondaryData.ResizeAllocation(0, 0, NumBytesPerElement);
 				}
 			}
 			else
 			{
-				if(!SecondaryData.GetAllocation())
+				if (!SecondaryData.GetAllocation())
 				{
 					// Allocate new indirect memory for the data.
-					SecondaryData.ResizeAllocation(0,NewMax,NumBytesPerElement);
+					SecondaryData.ResizeAllocation(0, NewMax, NumBytesPerElement);
 
 					// Move the data out of the inline data area into the new allocation.
 					RelocateConstructItems<ElementType>((void*)SecondaryData.GetAllocation(), GetInlineElements(), CurrentNum);
@@ -989,6 +989,41 @@ public:
 			}
 		}
 
+		void ResizeAllocation(SizeType CurrentNum, SizeType NewMax, SIZE_T NumBytesPerElement, uint32 AlignmentOfElement)
+		{
+			// Make sure the number of live elements is still withing the allocation since we only memmove, not destruct
+			checkSlow(CurrentNum <= NewMax);
+
+			// Check if the new allocation will fit in the inline data area.
+			if (NewMax <= NumInlineElements)
+			{
+				// If the old allocation wasn't in the inline data area, relocate it into the inline data area.
+				if (SecondaryData.GetAllocation())
+				{
+					RelocateConstructItems<ElementType>((void*)InlineData, (ElementType*)SecondaryData.GetAllocation(), CurrentNum);
+
+					// Free the old indirect allocation.
+					SecondaryData.ResizeAllocation(0, 0, NumBytesPerElement, AlignmentOfElement);
+				}
+			}
+			else
+			{
+				if (!SecondaryData.GetAllocation())
+				{
+					// Allocate new indirect memory for the data.
+					SecondaryData.ResizeAllocation(0, NewMax, NumBytesPerElement, AlignmentOfElement);
+
+					// Move the data out of the inline data area into the new allocation.
+					RelocateConstructItems<ElementType>((void*)SecondaryData.GetAllocation(), GetInlineElements(), CurrentNum);
+				}
+				else
+				{
+					// Reallocate the indirect data for the new size.
+					SecondaryData.ResizeAllocation(CurrentNum, NewMax, NumBytesPerElement, AlignmentOfElement);
+				}
+			}
+		}
+
 		UE_FORCEINLINE_HINT SizeType CalculateSlackReserve(SizeType NewMax, SIZE_T NumBytesPerElement) const
 		{
 			// If the elements use less space than the inline allocation, only use the inline allocation as slack.
@@ -996,6 +1031,13 @@ public:
 				NumInlineElements :
 				SecondaryData.CalculateSlackReserve(NewMax, NumBytesPerElement);
 		}
+		UE_FORCEINLINE_HINT SizeType CalculateSlackReserve(SizeType NewMax, SIZE_T NumBytesPerElement, uint32 AlignmentOfElement) const
+		{
+			// If the elements use less space than the inline allocation, only use the inline allocation as slack.
+			return NewMax <= NumInlineElements ?
+				NumInlineElements :
+				SecondaryData.CalculateSlackReserve(NewMax, NumBytesPerElement, AlignmentOfElement);
+		}
 		UE_FORCEINLINE_HINT SizeType CalculateSlackShrink(SizeType NewMax, SizeType CurrentMax, SIZE_T NumBytesPerElement) const
 		{
 			// If the elements use less space than the inline allocation, only use the inline allocation as slack.
@@ -1003,6 +1045,13 @@ public:
 				NumInlineElements :
 				SecondaryData.CalculateSlackShrink(NewMax, CurrentMax, NumBytesPerElement);
 		}
+		UE_FORCEINLINE_HINT SizeType CalculateSlackShrink(SizeType NewMax, SizeType CurrentMax, SIZE_T NumBytesPerElement, uint32 AlignmentOfElement) const
+		{
+			// If the elements use less space than the inline allocation, only use the inline allocation as slack.
+			return NewMax <= NumInlineElements ?
+				NumInlineElements :
+				SecondaryData.CalculateSlackShrink(NewMax, CurrentMax, NumBytesPerElement, AlignmentOfElement);
+		}
 		UE_FORCEINLINE_HINT SizeType CalculateSlackGrow(SizeType NewMax, SizeType CurrentMax, SIZE_T NumBytesPerElement) const
 		{
 			// If the elements use less space than the inline allocation, only use the inline allocation as slack.
@@ -1013,6 +1062,16 @@ public:
 				NumInlineElements :
 				SecondaryData.CalculateSlackGrow(NewMax, CurrentMax <= NumInlineElements ? 0 : CurrentMax, NumBytesPerElement);
 		}
+		UE_FORCEINLINE_HINT SizeType CalculateSlackGrow(SizeType NewMax, SizeType CurrentMax, SIZE_T NumBytesPerElement, uint32 AlignmentOfElement) const
+		{
+			// If the elements use less space than the inline allocation, only use the inline allocation as slack.
+			// Also, when computing slack growth, don't count inline elements -- the slack algorithm has a special
+			// case to save memory on the initial heap allocation, versus subsequent reallocations, and we don't
+			// want the inline elements to be treated as if they were the first heap allocation.
+			return NewMax <= NumInlineElements ?
+				NumInlineElements :
+				SecondaryData.CalculateSlackGrow(NewMax, CurrentMax <= NumInlineElements ? 0 : CurrentMax, NumBytesPerElement, AlignmentOfElement);
+		}
 
 		SIZE_T GetAllocatedSize(SizeType CurrentMax, SIZE_T NumBytesPerElement) const
 		{
@@ -1066,6 +1125,7 @@ public:
 template <uint32 NumInlineElements, int IndexSize, typename SecondaryAllocator>
 struct TAllocatorTraits<TSizedInlineAllocator<NumInlineElements, IndexSize, SecondaryAllocator>> : TAllocatorTraitsBase<TSizedInlineAllocator<NumInlineElements, IndexSize, SecondaryAllocator>>
 {
+	enum { SupportsElementAlignment = TAllocatorTraits<SecondaryAllocator>::SupportsElementAlignment };
 	enum { SupportsSlackTracking = true };
 };
 
diff --git a/Engine/Source/Runtime/Core/Public/Memory/VirtualLinearAllocator.h b/Engine/Source/Runtime/Core/Public/Memory/VirtualLinearAllocator.h
new file mode 100644
--- /dev/null
+++ b/Engine/Source/Runtime/Core/Public/Memory/VirtualLinearAllocator.h
@@ -0,0 +1,114 @@
+// Copyright Epic Games, Inc. All Rights Reserved.
+
+#pragma once
+
+#include "Async/UniqueLock.h"
+#include "HAL/Platform.h"
+#include "HAL/PlatformMemory.h"
+
+#ifndef UE_ENABLE_VIRTUAL_LINEAR_ALLOCATOR
+#define UE_ENABLE_VIRTUAL_LINEAR_ALLOCATOR PLATFORM_HAS_FPlatformVirtualMemoryBlock
+#endif
+
+#if UE_ENABLE_VIRTUAL_LINEAR_ALLOCATOR
+namespace UE
+{
+
+/**
+ * A thread-safe linear allocator backed by virtual memory.
+ *
+ * Reserves address space in one block of ReserveSize and commits pages on demand.
+ *
+ * Memory is decommitted by Reset() or when AddRef()/Release() are used to track
+ * references to the allocator and that reference count drops to 0.
+ *
+ * Reserved memory is released by the destructor. Use TNoDestroy to avoid freeing.
+ */
+template <typename MutexType>
+class alignas(PLATFORM_CACHE_LINE_SIZE) TVirtualLinearAllocator
+{
+	using FVirtualMemoryBlock = FPlatformMemory::FPlatformVirtualMemoryBlock;
+
+public:
+	// ReserveSize must be a multiple of commit alignment.
+	// It is currently safe on every platform to use multiples 64 KiB.
+	inline explicit TVirtualLinearAllocator(SIZE_T InReserveSize)
+		: ReserveSize(InReserveSize)
+	{
+		checkf(IsAligned(ReserveSize, FVirtualMemoryBlock::GetCommitAlignment()),
+			TEXT("ReserveSize %" SIZE_T_FMT " must be a multiple of commit alignment %" SIZE_T_FMT "."),
+			ReserveSize, FVirtualMemoryBlock::GetCommitAlignment());
+	}
+
+	inline ~TVirtualLinearAllocator()
+	{
+		Block.FreeVirtual();
+	}
+
+	void* Allocate(SIZE_T Size, uint32 Alignment)
+	{
+		TUniqueLock<MutexType> Lock(Mutex);
+
+		uint8* Base = (uint8*)Block.GetVirtualPointer();
+		if (UNLIKELY(!Base))
+		{
+			Block = FVirtualMemoryBlock::AllocateVirtual(ReserveSize);
+			Base = (uint8*)Block.GetVirtualPointer();
+			checkf(Base, TEXT("Failed to reserve %" SIZE_T_FMT " bytes of virtual memory in FVirtualLinearAllocator."), ReserveSize);
+		}
+
+		const SIZE_T AlignedOffset = Align(Offset, Alignment);
+		const SIZE_T NewOffset = AlignedOffset + Size;
+		checkf(NewOffset <= ReserveSize,
+			TEXT("Overflowed FVirtualLinearAllocator by %" SIZE_T_FMT " bytes, with reserved size of %" SIZE_T_FMT " bytes."),
+			NewOffset - ReserveSize, ReserveSize);
+
+		if (UNLIKELY(CommitSize < NewOffset))
+		{
+			const SIZE_T NewCommit = Align(NewOffset, FVirtualMemoryBlock::GetCommitAlignment());
+			Block.Commit(CommitSize, NewCommit - CommitSize);
+			CommitSize = NewCommit;
+		}
+
+		Offset = NewOffset;
+		return Base + AlignedOffset;
+	}
+
+	void Reset()
+	{
+		TUniqueLock<MutexType> Lock(Mutex);
+		// Only reset when the reference count is 0. This will always be the case when reference
+		// counting is not being used. When it is being used, AddRef() may be called on another
+		// thread between Release() seeing a reference count of 0 and Reset() locking the mutex.
+		if (ActiveCount.load(std::memory_order_relaxed) == 0 && CommitSize > 0)
+		{
+			Block.Decommit(0, CommitSize);
+			CommitSize = 0;
+			Offset = 0;
+		}
+	}
+
+	inline void AddRef()
+	{
+		ActiveCount.fetch_add(1, std::memory_order_relaxed);
+	}
+
+	inline void Release()
+	{
+		if (ActiveCount.fetch_sub(1, std::memory_order_acq_rel) == 1)
+		{
+			Reset();
+		}
+	}
+
+private:
+	FVirtualMemoryBlock Block;
+	SIZE_T ReserveSize = 0;
+	SIZE_T CommitSize = 0;
+	SIZE_T Offset = 0;
+	MutexType Mutex;
+	std::atomic<uint32> ActiveCount{0};
+};
+
+} // UE
+#endif // UE_ENABLE_VIRTUAL_LINEAR_ALLOCATOR
