[Backout] - CL51628270
#rnx
Original CL Desc
-----------------------------------------------------------------
ParkingLot: Use low-level virtual memory API for allocations to allow our allocators to depend on ParkingLot

Extended TSizedInlineAllocator with SupportsElementAlignment support.

#Synced-CL 50967057

diff --git a/Engine/Source/Runtime/Core/Private/Async/ParkingLot.cpp b/Engine/Source/Runtime/Core/Private/Async/ParkingLot.cpp
--- a/Engine/Source/Runtime/Core/Private/Async/ParkingLot.cpp
+++ b/Engine/Source/Runtime/Core/Private/Async/ParkingLot.cpp
@@ -9,16 +9,9 @@
 #include "Containers/Array.h"
 #include "Containers/ArrayView.h"
 #include "HAL/PlatformManualResetEvent.h"
-#include "Memory/VirtualLinearAllocator.h"
 #include "Templates/RefCounting.h"
-#include "Templates/NoDestroy.h"
 #include <atomic>
 
-#if !UE_ENABLE_VIRTUAL_LINEAR_ALLOCATOR
-#include "HAL/Allocators/AnsiAllocator.h"
-#include "HAL/MallocAnsi.h"
-#endif
-
 namespace UE::ParkingLot::Private
 {
 
@@ -51,188 +44,6 @@ static thread_local EExecutingFunction ParkingLotExecutingFunction = EExecutingF
 
 ///////////////////////////////////////////////////////////////////////////////////////////////////////////////////////
 
-#if UE_ENABLE_VIRTUAL_LINEAR_ALLOCATOR
-/**
- * Allocator for permanent allocations by ParkingLot: FTable, FBucket, FThread.
- *
- * Shared by FTable, TObjectAllocator<FBucket>, TObjectAllocator<FThread>.
- * ParkingLot allocates approximately ThreadCount buckets and threads, at 64 bytes each,
- * and tables totaling approximately 2*ThreadCount 8-byte pointers.
- * Even 16,384 threads will consume only 2.25 MiB, which makes 4 MiB plenty.
- *
- * Adds 64 KiB to avoid a multiple of 2 MiB as it can cause issues on platforms with transparent large pages.
- */
-static TNoDestroy<TVirtualLinearAllocator<FWordMutex>> ParkingLotPermanentAllocator(4 * 1024 * 1024 + 64 * 1024);
-#else
-static struct
-{
-	UE_REWRITE void* Allocate(SIZE_T Size, uint32 Alignment)
-	{
-		return AnsiMalloc(Size, Alignment);
-	}
-} ParkingLotPermanentAllocator;
-#endif
-
-///////////////////////////////////////////////////////////////////////////////////////////////////////////////////////
-
-/**
- * A low-level allocator that maintains a free list for a given object type. Memory is never released.
- */
-template <typename ObjectType>
-class TObjectAllocator
-{
-	static_assert(sizeof(ObjectType) >= sizeof(void*));
-	static_assert(alignof(ObjectType) >= alignof(void*));
-
-public:
-	static void* Allocate()
-	{
-		if (TUniqueLock Lock(FreeListMutex); void* Mem = FreeListHead)
-		{
-			FreeListHead = *(void**)Mem;
-			return Mem;
-		}
-		return ParkingLotPermanentAllocator->Allocate(sizeof(ObjectType), alignof(ObjectType));
-	}
-
-	static void Free(void* Mem)
-	{
-		TUniqueLock Lock(FreeListMutex);
-		void* Next = FreeListHead;
-		*(void**)Mem = Next;
-		FreeListHead = Mem;
-	}
-
-private:
-	inline static FWordMutex FreeListMutex;
-	inline static void* FreeListHead = nullptr;
-};
-
-///////////////////////////////////////////////////////////////////////////////////////////////////////////////////////
-
-#if UE_ENABLE_VIRTUAL_LINEAR_ALLOCATOR
-/**
- * Allocator for scratch allocations by ParkingLot: TArray<PointerType>.
- *
- * Scratch allocations are bounded by ThreadCount:
- * 1. FTable::Reserve can allocate maximum 3*ThreadCount 8-byte pointers while locked,
- *    and ThreadCount 8-byte pointers while unlocked. Reserve is only called while starting
- *    a new thread, which limits the number of concurrent calls. Even 256 new threads racing
- *    to increase the size to 1024 will allocate a worst-case ~2 MiB.
- * 2. WakeMultiple can allocate maximum WaitingThreadCount 8-byte pointers across every
- *    thread that is waking other threads. Wakes can race with timeouts to bring the
- *    effective total higher but even 512 threads waking 512 threads is only 2 MiB.
- *
- * Adds 64 KiB to avoid a multiple of 2 MiB as it can cause issues on platforms with transparent large pages.
- */
-inline static TNoDestroy<TVirtualLinearAllocator<FWordMutex>> ParkingLotScratchAllocator(4 * 1024 * 1024 + 64 * 1024);
-
-/**
- * A container allocator backed by a virtual linear allocator for scratch allocations.
- */
-class FScratchAllocator
-{
-public:
-	using SizeType = int32;
-	enum { NeedsElementType = false };
-	enum { RequireRangeCheck = false };
-	enum { ShrinkByDefault = false };
-
-	class ForAnyElementType
-	{
-	public:
-		constexpr ForAnyElementType() = default;
-
-		inline void MoveToEmpty(ForAnyElementType& Other)
-		{
-			Data = Other.Data;
-			Other.Data = nullptr;
-		}
-
-		inline void ResizeAllocation(SizeType CurrentNum, SizeType NewMax, SIZE_T NumBytesPerElement, uint32 AlignmentOfElement)
-		{
-			if (NewMax)
-			{
-				FScriptContainerElement* const OldData = Data;
-				Data = (FScriptContainerElement*)ParkingLotScratchAllocator->Allocate(NewMax * NumBytesPerElement, AlignmentOfElement);
-				if (OldData && CurrentNum > 0)
-				{
-					FMemory::Memcpy(Data, OldData, CurrentNum * NumBytesPerElement);
-				}
-			}
-			else
-			{
-				Data = nullptr;
-			}
-		}
-
-		inline SizeType CalculateSlackReserve(SizeType NewMax, SIZE_T NumBytesPerElement, uint32 AlignmentOfElement) const
-		{
-			return DefaultCalculateSlackReserve(NewMax, NumBytesPerElement, /*bAllowQuantize*/ false, AlignmentOfElement);
-		}
-
-		inline SizeType CalculateSlackShrink(SizeType NewMax, SizeType CurrentMax, SIZE_T NumBytesPerElement, uint32 AlignmentOfElement) const
-		{
-			return DefaultCalculateSlackShrink(NewMax, CurrentMax, NumBytesPerElement, /*bAllowQuantize*/ false, AlignmentOfElement);
-		}
-
-		inline SizeType CalculateSlackGrow(SizeType NewMax, SizeType CurrentMax, SIZE_T NumBytesPerElement, uint32 AlignmentOfElement) const
-		{
-			return DefaultCalculateSlackGrow(NewMax, CurrentMax, NumBytesPerElement, /*bAllowQuantize*/ false, AlignmentOfElement);
-		}
-
-		inline SIZE_T GetAllocatedSize(SizeType CurrentMax, SIZE_T NumBytesPerElement) const
-		{
-			return CurrentMax * NumBytesPerElement;
-		}
-
-		inline FScriptContainerElement* GetAllocation() const { return Data; }
-		inline bool HasAllocation() const { return !!Data; }
-
-		constexpr SizeType GetInitialCapacity() const { return 0; }
-
-	private:
-		ForAnyElementType(const ForAnyElementType&);
-		ForAnyElementType& operator=(const ForAnyElementType&);
-
-		FScriptContainerElement* Data = nullptr;
-	};
-
-	template <typename ElementType>
-	class ForElementType : public ForAnyElementType
-	{
-	public:
-		inline ElementType* GetAllocation() const
-		{
-			return (ElementType*)ForAnyElementType::GetAllocation();
-		}
-	};
-};
-#else
-static struct
-{
-	constexpr void AddRef() {}
-	constexpr void Release() {}
-} ParkingLotScratchAllocator;
-using FScratchAllocator = FAnsiAllocator;
-#endif // UE_ENABLE_VIRTUAL_LINEAR_ALLOCATOR
-
-} // UE::ParkingLot::Private
-
-#if UE_ENABLE_VIRTUAL_LINEAR_ALLOCATOR
-template <>
-struct TAllocatorTraits<UE::ParkingLot::Private::FScratchAllocator> : TAllocatorTraitsBase<UE::ParkingLot::Private::FScratchAllocator>
-{
-	enum { IsZeroConstruct = true };
-	enum { SupportsElementAlignment = true };
-};
-#endif
-
-///////////////////////////////////////////////////////////////////////////////////////////////////////////////////////
-
-namespace UE::ParkingLot::Private
-{
-
 /**
  * A thread as stored in the wait queue.
  */
@@ -253,28 +64,17 @@ struct alignas(PLATFORM_CACHE_LINE_SIZE) FThread final
 	{
 		if (ReferenceCount.fetch_sub(1, std::memory_order_acq_rel) == 1)
 		{
-			this->~FThread();
-			Allocator.Free(this);
+			delete this;
 		}
 	}
 
-	inline static FThread* New()
-	{
-		return new(Allocator.Allocate()) FThread;
-	}
-
-private:
 	FThread() = default;
 	~FThread() = default;
 
 	FThread(const FThread&) = delete;
 	FThread& operator=(const FThread&) = delete;
-
-	static TObjectAllocator<FThread> Allocator;
 };
 
-TObjectAllocator<FThread> FThread::Allocator;
-
 class FThreadLocalData
 {
 public:
@@ -288,9 +88,9 @@ public:
 		}
 		if (LocalThreadLocalData.bDestroyed)
 		{
-			return FThread::New();
+			return new FThread;
 		}
-		LocalThreadLocalData.Thread = FThread::New();
+		LocalThreadLocalData.Thread = new FThread;
 		return LocalThreadLocalData.Thread;
 	}
 
@@ -326,17 +126,6 @@ enum class EQueueAction
 class alignas(PLATFORM_CACHE_LINE_SIZE) FBucket final
 {
 public:
-	static FBucket* Create()
-	{
-		return new(Allocator.Allocate()) FBucket;
-	}
-
-	static void Destroy(FBucket* Bucket)
-	{
-		Bucket->~FBucket();
-		Allocator.Free(Bucket);
-	}
-
 	[[nodiscard]] inline TDynamicUniqueLock<FWordMutex> LockDynamic() { return TDynamicUniqueLock(Mutex); }
 
 	inline void Lock() { Mutex.Lock(); }
@@ -358,22 +147,18 @@ public:
 	template <typename VisitorType>
 	void DequeueIf(VisitorType&& Visitor);
 
-private:
 	FBucket() = default;
 	~FBucket() = default;
 
+private:
 	FBucket(const FBucket&) = delete;
 	FBucket& operator=(const FBucket&) = delete;
 
 	FWordMutex Mutex;
 	FThread* Head = nullptr;
 	FThread* Tail = nullptr;
-
-	static TObjectAllocator<FBucket> Allocator;
 };
 
-TObjectAllocator<FBucket> FBucket::Allocator;
-
 void FBucket::Enqueue(FThread* Thread)
 {
 	checkSlow(Thread);
@@ -476,7 +261,7 @@ private:
 	static void Destroy(FTable& Table);
 
 	/** Try to lock the whole table by locking every one of its buckets. */
-	static bool TryLock(FTable& Table, TArray<FBucket*, FScratchAllocator>& OutBuckets);
+	static bool TryLock(FTable& Table, TArray<FBucket*>& OutBuckets);
 	/** Unlock an array of buckets that was filled by TryLock. */
 	static void Unlock(TConstArrayView<FBucket*> LockedBuckets);
 
@@ -485,7 +270,6 @@ private:
 
 	/** Pointer to the current global table. Previous global tables are leaked. */
 	inline static std::atomic<FTable*> GlobalTable;
-	inline static FWordMutex CreateMutex;
 
 	FTable() = default;
 	~FTable() = default;
@@ -509,7 +293,7 @@ FBucket& FTable::FindOrCreateBucket(const void* Address, TDynamicUniqueLock<FWor
 	{
 		FTable& Table = CreateOrGet();
 		const uint32 Index = Hash % Table.BucketCount;
-		FBucket& Bucket = Table.FindOrCreateBucket(Index, [] { return FBucket::Create(); });
+		FBucket& Bucket = Table.FindOrCreateBucket(Index, [] { return new FBucket; });
 		OutLock = Bucket.LockDynamic();
 
 		if (LIKELY(&Table == GlobalTable.load(std::memory_order_acquire)))
@@ -563,7 +347,7 @@ FBucket& FTable::FindOrCreateBucket(uint32 Index, AllocatorType&& BucketAllocato
 		}
 		else
 		{
-			FBucket::Destroy(NewBucket);
+			delete NewBucket;
 		}
 		checkSlow(Bucket);
 	}
@@ -572,27 +356,30 @@ FBucket& FTable::FindOrCreateBucket(uint32 Index, AllocatorType&& BucketAllocato
 
 FTable& FTable::CreateOrGet()
 {
-	if (FTable* Table = GlobalTable.load(std::memory_order_acquire); LIKELY(Table))
+	FTable* Table = GlobalTable.load(std::memory_order_acquire);
+
+	if (LIKELY(Table))
 	{
 		return *Table;
 	}
 
-	TUniqueLock Lock(CreateMutex);
-	if (FTable* Table = GlobalTable.load(std::memory_order_acquire))
+	FTable& NewTable = Create(MinSize);
+
+	if (LIKELY(GlobalTable.compare_exchange_strong(Table, &NewTable, std::memory_order_release, std::memory_order_acquire)))
 	{
-		return *Table;
+		return NewTable;
 	}
 
-	FTable& NewTable = Create(MinSize);
-	FTable* CompareTable = GlobalTable.exchange(&NewTable, std::memory_order_release);
-	checkSlow(!CompareTable);
-	return NewTable;
+	Destroy(NewTable);
+
+	checkSlow(Table);
+	return *Table;
 }
 
 FTable& FTable::Create(const uint32 Size)
 {
-	const SIZE_T MemorySize = sizeof(FTable) + sizeof(std::atomic<FBucket*>) * SIZE_T(Size);
-	void* const Memory = ParkingLotPermanentAllocator->Allocate(MemorySize, alignof(FTable));
+	const SIZE_T MemorySize = sizeof(FTable) + sizeof(FBucket*) * SIZE_T(Size);
+	void* const Memory = FMemory::Malloc(MemorySize, alignof(FTable));
 	FMemory::Memzero(Memory, MemorySize);
 	FTable& Table = *new(Memory) FTable;
 	Table.BucketCount = Size;
@@ -602,15 +389,13 @@ FTable& FTable::Create(const uint32 Size)
 void FTable::Destroy(FTable& Table)
 {
 	Table.~FTable();
+	FMemory::Free(&Table);
 }
 
 void FTable::Reserve(const uint32 ThreadCount)
 {
-	// Reference the allocator to allow it to decommit scratch memory.
-	TRefCountPtr ScratchAllocatorRef(&ParkingLotScratchAllocator);
-
 	const uint32 TargetBucketCount = FMath::RoundUpToPowerOfTwo(ThreadCount);
-	TArray<FBucket*, FScratchAllocator> ExistingBuckets;
+	TArray<FBucket*> ExistingBuckets;
 
 	for (;;)
 	{
@@ -630,7 +415,7 @@ void FTable::Reserve(const uint32 ThreadCount)
 
 		// Gather waiting threads to be redistributed into the buckets of the new table.
 		// Threads with the same address remain in the same relative order as they were queued.
-		TArray<FThread*, TInlineAllocator<128, FScratchAllocator>> Threads;
+		TArray<FThread*> Threads;
 		for (FBucket* Bucket : ExistingBuckets)
 		{
 			while (FThread* Thread = Bucket->Dequeue())
@@ -642,10 +427,10 @@ void FTable::Reserve(const uint32 ThreadCount)
 		FTable& NewTable = Create(TargetBucketCount);
 
 		// Reuse existing now-empty buckets when populating the new table.
-		TArray<FBucket*, FScratchAllocator> AvailableBuckets = ExistingBuckets;
+		TArray<FBucket*> AvailableBuckets = ExistingBuckets;
 		const auto AllocateBucket = [&AvailableBuckets]() -> FBucket*
 		{
-			return !AvailableBuckets.IsEmpty() ? AvailableBuckets.Pop(EAllowShrinking::No) : FBucket::Create();
+			return !AvailableBuckets.IsEmpty() ? AvailableBuckets.Pop(EAllowShrinking::No) : new FBucket;
 		};
 
 		// Add waiting threads to the new table.
@@ -674,14 +459,14 @@ void FTable::Reserve(const uint32 ThreadCount)
 	}
 }
 
-bool FTable::TryLock(FTable& Table, TArray<FBucket*, FScratchAllocator>& OutBuckets)
+bool FTable::TryLock(FTable& Table, TArray<FBucket*>& OutBuckets)
 {
 	OutBuckets.Reset(Table.BucketCount);
 
 	// Gather buckets from the table, creating them as needed because the lock is on the bucket.
 	for (uint32 Index = 0; Index < Table.BucketCount; ++Index)
 	{
-		OutBuckets.Add(&Table.FindOrCreateBucket(Index, [] { return FBucket::Create(); }));
+		OutBuckets.Add(&Table.FindOrCreateBucket(Index, [] { return new FBucket; }));
 	}
 
 	// Lock the buckets in order by address to ensure consistent ordering regardless of the table being locked.
@@ -971,9 +756,7 @@ uint32 WakeMultiple(const void* const Address, const uint32 WakeCount)
 	checkf(ParkingLotExecutingFunction <= EExecutingFunction::BeforeWait,
 		TEXT("Detected a re-entrant call into ParkingLot from within a %s callback."), LexToString(ParkingLotExecutingFunction));
 
-	// Reference the allocator to allow it to decommit scratch memory.
-	TRefCountPtr ScratchAllocatorRef(&ParkingLotScratchAllocator);
-	TArray<TRefCountPtr<FThread>, TInlineAllocator<128, FScratchAllocator>> WakeThreads;
+	TArray<TRefCountPtr<FThread>, TInlineAllocator<128>> WakeThreads;
 
 	if (TDynamicUniqueLock<FWordMutex> BucketLock; FBucket* Bucket = FTable::FindBucket(Address, BucketLock))
 	{
diff --git a/Engine/Source/Runtime/Core/Public/Containers/ContainerAllocationPolicies.h b/Engine/Source/Runtime/Core/Public/Containers/ContainerAllocationPolicies.h
--- a/Engine/Source/Runtime/Core/Public/Containers/ContainerAllocationPolicies.h
+++ b/Engine/Source/Runtime/Core/Public/Containers/ContainerAllocationPolicies.h
@@ -954,29 +954,29 @@ public:
 			return GetInlineElements();
 		}
 
-		void ResizeAllocation(SizeType CurrentNum, SizeType NewMax, SIZE_T NumBytesPerElement)
+		void ResizeAllocation(SizeType CurrentNum, SizeType NewMax,SIZE_T NumBytesPerElement)
 		{
 			// Make sure the number of live elements is still withing the allocation since we only memmove, not destruct
 			checkSlow(CurrentNum <= NewMax);
 
 			// Check if the new allocation will fit in the inline data area.
-			if (NewMax <= NumInlineElements)
+			if(NewMax <= NumInlineElements)
 			{
 				// If the old allocation wasn't in the inline data area, relocate it into the inline data area.
-				if (SecondaryData.GetAllocation())
+				if(SecondaryData.GetAllocation())
 				{
 					RelocateConstructItems<ElementType>((void*)InlineData, (ElementType*)SecondaryData.GetAllocation(), CurrentNum);
 
 					// Free the old indirect allocation.
-					SecondaryData.ResizeAllocation(0, 0, NumBytesPerElement);
+					SecondaryData.ResizeAllocation(0,0,NumBytesPerElement);
 				}
 			}
 			else
 			{
-				if (!SecondaryData.GetAllocation())
+				if(!SecondaryData.GetAllocation())
 				{
 					// Allocate new indirect memory for the data.
-					SecondaryData.ResizeAllocation(0, NewMax, NumBytesPerElement);
+					SecondaryData.ResizeAllocation(0,NewMax,NumBytesPerElement);
 
 					// Move the data out of the inline data area into the new allocation.
 					RelocateConstructItems<ElementType>((void*)SecondaryData.GetAllocation(), GetInlineElements(), CurrentNum);
@@ -989,41 +989,6 @@ public:
 			}
 		}
 
-		void ResizeAllocation(SizeType CurrentNum, SizeType NewMax, SIZE_T NumBytesPerElement, uint32 AlignmentOfElement)
-		{
-			// Make sure the number of live elements is still withing the allocation since we only memmove, not destruct
-			checkSlow(CurrentNum <= NewMax);
-
-			// Check if the new allocation will fit in the inline data area.
-			if (NewMax <= NumInlineElements)
-			{
-				// If the old allocation wasn't in the inline data area, relocate it into the inline data area.
-				if (SecondaryData.GetAllocation())
-				{
-					RelocateConstructItems<ElementType>((void*)InlineData, (ElementType*)SecondaryData.GetAllocation(), CurrentNum);
-
-					// Free the old indirect allocation.
-					SecondaryData.ResizeAllocation(0, 0, NumBytesPerElement, AlignmentOfElement);
-				}
-			}
-			else
-			{
-				if (!SecondaryData.GetAllocation())
-				{
-					// Allocate new indirect memory for the data.
-					SecondaryData.ResizeAllocation(0, NewMax, NumBytesPerElement, AlignmentOfElement);
-
-					// Move the data out of the inline data area into the new allocation.
-					RelocateConstructItems<ElementType>((void*)SecondaryData.GetAllocation(), GetInlineElements(), CurrentNum);
-				}
-				else
-				{
-					// Reallocate the indirect data for the new size.
-					SecondaryData.ResizeAllocation(CurrentNum, NewMax, NumBytesPerElement, AlignmentOfElement);
-				}
-			}
-		}
-
 		UE_FORCEINLINE_HINT SizeType CalculateSlackReserve(SizeType NewMax, SIZE_T NumBytesPerElement) const
 		{
 			// If the elements use less space than the inline allocation, only use the inline allocation as slack.
@@ -1031,13 +996,6 @@ public:
 				NumInlineElements :
 				SecondaryData.CalculateSlackReserve(NewMax, NumBytesPerElement);
 		}
-		UE_FORCEINLINE_HINT SizeType CalculateSlackReserve(SizeType NewMax, SIZE_T NumBytesPerElement, uint32 AlignmentOfElement) const
-		{
-			// If the elements use less space than the inline allocation, only use the inline allocation as slack.
-			return NewMax <= NumInlineElements ?
-				NumInlineElements :
-				SecondaryData.CalculateSlackReserve(NewMax, NumBytesPerElement, AlignmentOfElement);
-		}
 		UE_FORCEINLINE_HINT SizeType CalculateSlackShrink(SizeType NewMax, SizeType CurrentMax, SIZE_T NumBytesPerElement) const
 		{
 			// If the elements use less space than the inline allocation, only use the inline allocation as slack.
@@ -1045,13 +1003,6 @@ public:
 				NumInlineElements :
 				SecondaryData.CalculateSlackShrink(NewMax, CurrentMax, NumBytesPerElement);
 		}
-		UE_FORCEINLINE_HINT SizeType CalculateSlackShrink(SizeType NewMax, SizeType CurrentMax, SIZE_T NumBytesPerElement, uint32 AlignmentOfElement) const
-		{
-			// If the elements use less space than the inline allocation, only use the inline allocation as slack.
-			return NewMax <= NumInlineElements ?
-				NumInlineElements :
-				SecondaryData.CalculateSlackShrink(NewMax, CurrentMax, NumBytesPerElement, AlignmentOfElement);
-		}
 		UE_FORCEINLINE_HINT SizeType CalculateSlackGrow(SizeType NewMax, SizeType CurrentMax, SIZE_T NumBytesPerElement) const
 		{
 			// If the elements use less space than the inline allocation, only use the inline allocation as slack.
@@ -1062,16 +1013,6 @@ public:
 				NumInlineElements :
 				SecondaryData.CalculateSlackGrow(NewMax, CurrentMax <= NumInlineElements ? 0 : CurrentMax, NumBytesPerElement);
 		}
-		UE_FORCEINLINE_HINT SizeType CalculateSlackGrow(SizeType NewMax, SizeType CurrentMax, SIZE_T NumBytesPerElement, uint32 AlignmentOfElement) const
-		{
-			// If the elements use less space than the inline allocation, only use the inline allocation as slack.
-			// Also, when computing slack growth, don't count inline elements -- the slack algorithm has a special
-			// case to save memory on the initial heap allocation, versus subsequent reallocations, and we don't
-			// want the inline elements to be treated as if they were the first heap allocation.
-			return NewMax <= NumInlineElements ?
-				NumInlineElements :
-				SecondaryData.CalculateSlackGrow(NewMax, CurrentMax <= NumInlineElements ? 0 : CurrentMax, NumBytesPerElement, AlignmentOfElement);
-		}
 
 		SIZE_T GetAllocatedSize(SizeType CurrentMax, SIZE_T NumBytesPerElement) const
 		{
@@ -1125,7 +1066,6 @@ public:
 template <uint32 NumInlineElements, int IndexSize, typename SecondaryAllocator>
 struct TAllocatorTraits<TSizedInlineAllocator<NumInlineElements, IndexSize, SecondaryAllocator>> : TAllocatorTraitsBase<TSizedInlineAllocator<NumInlineElements, IndexSize, SecondaryAllocator>>
 {
-	enum { SupportsElementAlignment = TAllocatorTraits<SecondaryAllocator>::SupportsElementAlignment };
 	enum { SupportsSlackTracking = true };
 };
 
diff --git a/Engine/Source/Runtime/Core/Public/Memory/VirtualLinearAllocator.h b/Engine/Source/Runtime/Core/Public/Memory/VirtualLinearAllocator.h
deleted file mode 100644
