Clean up comments on TNonNullPtr to make it clear that they can be uninitialized during object creation and add ResetToUninitialized to correctly clear the state during shutdown and object recycling

#Synced-CL 51587375

diff --git a/Engine/Source/Runtime/Core/Public/Templates/NonNullPointer.h b/Engine/Source/Runtime/Core/Public/Templates/NonNullPointer.h
--- a/Engine/Source/Runtime/Core/Public/Templates/NonNullPointer.h
+++ b/Engine/Source/Runtime/Core/Public/Templates/NonNullPointer.h
@@ -13,10 +13,17 @@
 #include <type_traits>
 
 class FArchive;
-enum class EDefaultConstructNonNullPtr { UnsafeDoNotUse }; // So we can construct TNonNullPtrs
+enum class EDefaultConstructNonNullPtr 
+{ 
+	Uninitialized,					// Allocate as specifically uninitialized, will return false for IsInitialized()
+	UnsafeDoNotUse = Uninitialized	// Old name for Uninitialized, to be deprecated
+};
 
 /**
- * TNonNullPtr is a non-nullable, non-owning, raw/naked/unsafe pointer.
+ * TNonNullPtr is a non-nullable pointer to ObjectType that can be used to track references that should never be null.
+ * It will never be null during normal operation, but may be uninitialized during object construction and destruction.
+ * It supports intrusive optional state so TOptional<TNonNullPtr<ObjectType>> stores an unset optional pointer as uninitialized.
+ * If you need a simple template that checks for null in development with no extra behavior, use TNotNull<ObjectType*> instead.
  */
 template<typename ObjectType>
 class TNonNullPtr
@@ -24,7 +31,7 @@ class TNonNullPtr
 public:
 
 	/**
-	 * Hack that can be used under extraordinary circumstances
+	 * Start as uninitialized, so it can be used as a member variable and set during later object initialization
 	 */
 	UE_FORCEINLINE_HINT TNonNullPtr(EDefaultConstructNonNullPtr)
 		: Object(nullptr)
@@ -211,15 +218,23 @@ public:
 		return Object;
 	}
 
-	/*
-	 * WARNING: Hack that can be used under extraordinary circumstances. Pointers here 
-	 * should always be valid but might be in the EDefaultConstructNonNullPtr state 
-	 * during initialization.
+	/**
+	 * Check to see if a pointer has been initialized, which means it will point to a valid object for the lifecycle of the owning object.
+	 * Pointers here should always be valid but might be in the EDefaultConstructNonNullPtr state during initialization.
 	 */
 	UE_FORCEINLINE_HINT bool IsInitialized() const
 	{
 		return Object != nullptr;
 	}
+
+	/**
+	 * WARNING: Hack that can be used under extraordinary circumstances. Explicitly reset pointer to uninitialized state
+	 * This should only be used during object shutdown or recycling to avoid dangling pointers.
+	 */
+	inline void ResetToUninitialized()
+	{
+		Object = nullptr;
+	}
 	
 	/**
 	 * Use IsInitialized if needed
diff --git a/Engine/Source/Runtime/CoreUObject/Public/UObject/ScriptMacros.h b/Engine/Source/Runtime/CoreUObject/Public/UObject/ScriptMacros.h
--- a/Engine/Source/Runtime/CoreUObject/Public/UObject/ScriptMacros.h
+++ b/Engine/Source/Runtime/CoreUObject/Public/UObject/ScriptMacros.h
@@ -144,7 +144,7 @@ enum {MAX_VARIABLE_SIZE = 0x0FFF };
 #define P_GET_VERSECASTABLECONCRETESUBTYPE_REF(ElementType, ParamName)	PARAM_PASSED_BY_REF(ParamName, FObjectPropertyBase, verse::castable_concrete_subtype<ElementType>)
 
 
-#define P_GET_INTERFACEINSTANCE(ElementType, ParamName)		PARAM_PASSED_BY_VAL_INITED(ParamName, FObjectPropertyBase, TInterfaceInstance<ElementType>, EDefaultConstructNonNullPtr::UnsafeDoNotUse)
+#define P_GET_INTERFACEINSTANCE(ElementType, ParamName)		PARAM_PASSED_BY_VAL_INITED(ParamName, FObjectPropertyBase, TInterfaceInstance<ElementType>, EDefaultConstructNonNullPtr::Uninitialized)
 #define P_GET_INTERFACEINSTANCE_REF(ElementType, ParamName)	PARAM_PASSED_BY_REF(ParamName, FObjectPropertyBase, TInterfaceInstance<ElementType>)
 
 #define P_GET_TOPTIONAL(ElementType,ParamName)		PARAM_PASSED_BY_VAL(ParamName, FOptionalProperty, TOptional<ElementType>)
