LowLevelMemTracker: FIx corrupt csv values due to mismatched header column and rows in the FLLMCsvWriter when tracking the CodeVsContent tagset.
From Matt Peters

diff --git a/Engine/Source/Runtime/Core/Private/HAL/LLM/LLM.cpp b/Engine/Source/Runtime/Core/Private/HAL/LLM/LLM.cpp
--- a/Engine/Source/Runtime/Core/Private/HAL/LLM/LLM.cpp
+++ b/Engine/Source/Runtime/Core/Private/HAL/LLM/LLM.cpp
@@ -6919,6 +6919,7 @@ static bool CalculateCodeOrContentColumns(TLLMMap<CsvWriter::FColumnKey, CsvWrit
 		if (!ColumnData.bInitialized)
 		{
 			ColumnData.bInitialized = true;
+			ColumnData.Index = Columns.Num() - 1;
 			bUpdatedColumns = true;
 		}
 		ColumnData.Size += CurrentSize;
@@ -7132,10 +7133,18 @@ bool FLLMCsvWriter::UpdateColumns(FLLMGlobals& LLMRef, const FTrackerTagSizeMap&
 		{
 			bUpdated = true;
 			ColumnData.bInitialized = true;
+			ColumnData.Index = Columns.Num() - 1;
 		}
 		// ColumnData.Size is unused when !bRecordingCodeOrContent, because we already have
 		// the accumulated value in TagSizes and do not need to accumulate a sum.
 	}
+	if (bUpdated)
+	{
+		Columns.ValueSort([](const CsvWriter::FColumnData& A, const CsvWriter::FColumnData& B)
+			{
+				return A.Index < B.Index;
+			});
+	}
 	return bUpdated;
 }
 
@@ -7232,17 +7241,16 @@ void FLLMCsvWriter::AddRow(FLLMGlobals& LLMRef, const FTrackerTagSizeMap& TagSiz
 		WriteTag(TagToReadFromForUntagged);
 	}
 
-	if (bRecordingCodeOrContent)
+	for (TPair<CsvWriter::FColumnKey, CsvWriter::FColumnData>& Pair : Columns)
 	{
-		for (TPair<CsvWriter::FColumnKey, CsvWriter::FColumnData>& Pair : Columns)
+		if (Pair.Key.CodeOrContentTag)
 		{
 			WriteValue(Pair.Value.Size);
 		}
-	}
-
-	for (TPair<CsvWriter::FColumnKey, CsvWriter::FColumnData>& Pair : Columns)
-	{
-		WriteTag(Pair.Key.SystemsTag);
+		else
+		{
+			WriteTag(Pair.Key.SystemsTag);
+		}
 	}
 	Write(TEXTVIEW("\n"));
 
@@ -7503,7 +7511,14 @@ void FLLMCsvProfilerWriter::Publish(FLLMGlobals& LLMRef, const FTrackerTagSizeMa
 
 	if (bRecordingCodeOrContent)
 	{
-		CalculateCodeOrContentColumns(Columns, Tracker, LLMRef, LLMRef.ActiveTagsToAllocationGroup, LLMRef.AllocationGroupsMutex);
+		bool bUpdated = CalculateCodeOrContentColumns(Columns, Tracker, LLMRef, LLMRef.ActiveTagsToAllocationGroup, LLMRef.AllocationGroupsMutex);
+		if (bUpdated)
+		{
+			Columns.ValueSort([](const CsvWriter::FColumnData& A, const CsvWriter::FColumnData& B)
+				{
+					return A.Index < B.Index;
+				});
+		}
 
 		for (const TPair<CsvWriter::FColumnKey, CsvWriter::FColumnData>& KVPair : Columns)
 		{
diff --git a/Engine/Source/Runtime/Core/Private/HAL/LLM/LLMPrivate.h b/Engine/Source/Runtime/Core/Private/HAL/LLM/LLMPrivate.h
--- a/Engine/Source/Runtime/Core/Private/HAL/LLM/LLMPrivate.h
+++ b/Engine/Source/Runtime/Core/Private/HAL/LLM/LLMPrivate.h
@@ -45,6 +45,7 @@ struct FColumnKey
 struct FColumnData
 {
 	int64 Size = 0;
+	int32 Index = 0;
 	bool bInitialized = false;
 };
 
