Fixed warning when loading missing LocMeta/LocRes files

#rnx
#Synced-CL 51656536

diff --git a/Engine/Source/Runtime/Core/Private/Internationalization/TextLocalizationResource.cpp b/Engine/Source/Runtime/Core/Private/Internationalization/TextLocalizationResource.cpp
--- a/Engine/Source/Runtime/Core/Private/Internationalization/TextLocalizationResource.cpp
+++ b/Engine/Source/Runtime/Core/Private/Internationalization/TextLocalizationResource.cpp
@@ -63,7 +63,7 @@ bool FTextLocalizationMetaDataResource::LoadFromFile(const FString& FilePath)
 	TUniquePtr<FArchive> Reader;
 #if PRELOAD_LOCMETA_FILES
 	TArray<uint8> FileBytes;
-	if (FFileHelper::LoadFileToArray(FileBytes, *FilePath))
+	if (FFileHelper::LoadFileToArray(FileBytes, *FilePath, FILEREAD_Silent))
 	{
 		Reader = MakeUnique<FMemoryReader>(FileBytes);
 	}
@@ -257,7 +257,7 @@ bool FTextLocalizationResource::LoadFromFile(const FString& FilePath, const int3
 	TUniquePtr<FArchive> Reader;
 #if PRELOAD_LOCRES_FILES
 	TArray<uint8> FileBytes;
-	if (FFileHelper::LoadFileToArray(FileBytes, *FilePath))
+	if (FFileHelper::LoadFileToArray(FileBytes, *FilePath, FILEREAD_Silent))
 	{
 		Reader = MakeUnique<FMemoryReader>(FileBytes);
 	}
