Add WindowsPlatform.DisableMessageBox to prevent MessageBox display

#rnx

diff --git a/Engine/Source/Runtime/Core/Private/Windows/WindowsPlatformMisc.cpp b/Engine/Source/Runtime/Core/Private/Windows/WindowsPlatformMisc.cpp
--- a/Engine/Source/Runtime/Core/Private/Windows/WindowsPlatformMisc.cpp
+++ b/Engine/Source/Runtime/Core/Private/Windows/WindowsPlatformMisc.cpp
@@ -126,6 +126,15 @@ static TAutoConsoleVariable<int32> CVarDriverDetectionMethod(
 	#define UE_CPUID(Info, FunctionID) __cpuid(Info, FunctionID);
 #endif //PLATFORM_COMPILER_CLANG
 
+namespace UE::Core::Private
+{
+	static TAutoConsoleVariable<bool> CVarDisableWindowsMessageBox(
+		TEXT("WindowsPlatform.DisableMessageBox"),
+		false,
+		TEXT("If this is true, no MessageBox will be displayed and calls to MessageBoxExt will return EAppReturnType::Cancel by default."),
+		ECVF_Default);
+}
+
 int32 GetOSVersionsHelper( TCHAR* OutOSVersionLabel, int32 OSVersionLabelLength, TCHAR* OutOSSubVersionLabel, int32 OSSubVersionLabelLength )
 {
 	int32 ErrorCode = (int32)FWindowsOSVersionHelper::SUCCEEDED;
@@ -2103,6 +2112,12 @@ public:
 
 EAppReturnType::Type FWindowsPlatformMisc::MessageBoxExt( EAppMsgType::Type MsgType, const TCHAR* Text, const TCHAR* Caption )
 {
+	if (UE::Core::Private::CVarDisableWindowsMessageBox.GetValueOnAnyThread() == true)
+	{
+		UE_LOG(LogWindows, Warning, TEXT("FWindowsPlatformMisc::MessageBoxExt: A popup creation was auto-cancelled (%s)"), Caption);
+		return EAppReturnType::Cancel;
+	}
+
 	FSlowHeartBeatScope SuspendHeartBeat;
 
 	HWND ParentWindow = (HWND)NULL;
