Fix AutoRTFM::PrintCallstackOnCrash not printing callstacks for UE assertion failures during static destruction
- Assertion failures during normal execution printed error messages fine because UE's log infrastructure (GLog, GError) was still alive. However, assertions firing during CRT static destruction (e.g. in destructors of file-scope statics) silently terminated the process with exit code 3 and no diagnostic output, because GLog/GError may already be torn down at that point.
- Fix: PrintCallstackOnCrash now sets GIsGuarded=true so assertion failures raise SEH exception 0x4000 via ReportAssert instead of going through HandleError -> TerminateProcess(3). The VEH handler (which persists through static destruction and writes directly to stderr via fprintf) now also handles exception code 0x4000.
#Synced-CL 51688190

diff --git a/Engine/Source/Runtime/AutoRTFM/Private/TestUtils.cpp b/Engine/Source/Runtime/AutoRTFM/Private/TestUtils.cpp
--- a/Engine/Source/Runtime/AutoRTFM/Private/TestUtils.cpp
+++ b/Engine/Source/Runtime/AutoRTFM/Private/TestUtils.cpp
@@ -344,6 +344,15 @@ AUTORTFM_DISABLE static LONG WINAPI CrashHandler(_EXCEPTION_POINTERS* Info)
 		case EXCEPTION_STACK_OVERFLOW:
 			Kind = "stack overflow";
 			break;
+		case 0x4000:  // see AssertExceptionCode in WindowsPlatformCrashContext.cpp
+			Kind = "UE assertion failed";
+			break;
+		case 0x8000:  // see GPUCrashExceptionCode in WindowsPlatformCrashContext.cpp
+			Kind = "UE GPU crash";
+			break;
+		case 0xc000:  // see OutOfMemoryExceptionCode in WindowsPlatformCrashContext.cpp
+			Kind = "UE out of memory error";
+			break;
 		default:
 			return EXCEPTION_CONTINUE_SEARCH;
 	}
diff --git a/Engine/Source/Runtime/Core/Private/AutoRTFM/AutoRTFMTestingUE.cpp b/Engine/Source/Runtime/Core/Private/AutoRTFM/AutoRTFMTestingUE.cpp
--- a/Engine/Source/Runtime/Core/Private/AutoRTFM/AutoRTFMTestingUE.cpp
+++ b/Engine/Source/Runtime/Core/Private/AutoRTFM/AutoRTFMTestingUE.cpp
@@ -2,6 +2,7 @@
 
 #include "AutoRTFM/AutoRTFMTestingUE.h"
 
+#include "CoreGlobals.h"
 #include "HAL/PlatformStackWalk.h"
 
 #if AUTORTFM_ENABLE_TEST_UTILS
@@ -26,6 +27,10 @@ void InterceptExternAPIs([[maybe_unused]] TFunction<void(autortfm_extern_api&)>
 
 void PrintCallstackOnCrash()
 {
+	// Enable guarded execution so that check/ensure failures raise SEH exceptions
+	// (caught by the VEH handler installed below) instead of logging and terminating the process.
+	GIsGuarded = true;
+
 	OnCrash([](const FCrashInfo& Info)
 	{
 		static constexpr uint32 MaxCallstackDepth = 256;
