Remove RF_Public flag from runtime entity and component creation to reduce package export overhead and improve load times.
- Objects marked RF_Public are exported in package tables, increasing package size and load time.
- Runtime entities and components don't need to be publicly exported since we allow references between private objects that share the same outer across packages.
- However, archetype entities/components (prefabs) retain RF_Public via GetMaskedFlags(RF_ArchetypeObject) for proper asset serialization and referencing.
- Refactored entity::CreateComponentNative to use UE::SceneGraphAPI::ConstructComponent for consistent component creation. Duplicate component types now return nullptr instead of crashing.
- Tightened ConstructEntity signature to take TNotNull<verse::entity*> instead of nullable pointer.
- Added editor-only ensureMsgf in entity and component constructors to catch accidental RF_Public usage going forward.
- Added serialization fixup to strip RF_Public from entities and subobjects saved before this change.
- Fixed entity::Rename to no longer gate redirector creation on RF_Public, since entities are no longer public.

#Synced-CL 51708345
#Synced-CL 51708345

#Synced-CL 51708345

diff --git a/Engine/Source/Runtime/Core/Public/UObject/FortniteMainBranchObjectVersions.inl b/Engine/Source/Runtime/Core/Public/UObject/FortniteMainBranchObjectVersions.inl
--- a/Engine/Source/Runtime/Core/Public/UObject/FortniteMainBranchObjectVersions.inl
+++ b/Engine/Source/Runtime/Core/Public/UObject/FortniteMainBranchObjectVersions.inl
@@ -784,5 +784,8 @@ FORTNITE_MAIN_VERSION(Composite_LayerNameDeprecation, 256)
 // FDynamicMeshAttributeSet has triangle labels
 FORTNITE_MAIN_VERSION(DynamicMeshAttributesTriangleLabels, 257)
 
+// Remove RF_Public on entity components and subobjects
+FORTNITE_MAIN_VERSION(FixPublicEntityComponentsAndSubobjects, 258)
+
 // -----<new versions can be added above this line>-------------------------------------------------
 #undef FORTNITE_MAIN_VERSION
