Fix destruction order of FTlsAutoCleanup to match the behavior of thread_local and allow usage for inter-dependent objects

#rnx
#Synced-CL 51619035

diff --git a/Engine/Source/Runtime/Core/Private/HAL/ThreadingBase.cpp b/Engine/Source/Runtime/Core/Private/HAL/ThreadingBase.cpp
--- a/Engine/Source/Runtime/Core/Private/HAL/ThreadingBase.cpp
+++ b/Engine/Source/Runtime/Core/Private/HAL/ThreadingBase.cpp
@@ -1575,7 +1575,19 @@ FTlsAutoCleanup* FThreadSingletonInitializer::Inject(FTlsAutoCleanup* Instance,
 
 void FTlsAutoCleanup::Register()
 {
-	static thread_local TArray<TUniquePtr<FTlsAutoCleanup>> TlsInstances;
+	// Destruction of TLS values need to occur in the reverse order of initialization, as thread_local does.
+	struct FTlsContainer : public TArray<TUniquePtr<FTlsAutoCleanup>>
+	{
+		~FTlsContainer()
+		{
+			for (int32 Index = Num() - 1; Index >= 0; Index--)
+			{
+				(*this)[Index].Reset();
+			}
+		}
+	};
+
+	static thread_local FTlsContainer TlsInstances;
 	TlsInstances.Add(TUniquePtr<FTlsAutoCleanup>(this));
 }
 
diff --git a/Engine/Source/Runtime/Core/Tests/HAL/ThreadTest.cpp b/Engine/Source/Runtime/Core/Tests/HAL/ThreadTest.cpp
--- a/Engine/Source/Runtime/Core/Tests/HAL/ThreadTest.cpp
+++ b/Engine/Source/Runtime/Core/Tests/HAL/ThreadTest.cpp
@@ -119,17 +119,16 @@ namespace
 			CHECK_MESSAGE(TEXT("Thread singleton should be uninitialized by default"), bDefaultValuePass);
 			CHECK_MESSAGE(TEXT("Thread singletons in different threads should be isolated"), FThreadSingletonTest::Get().GetTestField() == 1);
 		}
-		{	// check that ThreadSingleton entries don't point to invalid memory after cleanup
-			class FThreadSingletonFirst : public TThreadSingleton<FThreadSingletonFirst>
+		{	// check that ThreadSingleton entries don't point to invalid memory after cleanup, and that destruction occurs in the reverse order of construction.
+			class FThreadSingletonSecond : public TThreadSingleton<FThreadSingletonSecond>
 			{
 			};
-			class FThreadSingletonSecond : public TThreadSingleton<FThreadSingletonSecond>
+			class FThreadSingletonFirst : public TThreadSingleton<FThreadSingletonFirst>
 			{
-			public:
-				virtual ~FThreadSingletonSecond()
+				virtual ~FThreadSingletonFirst()
 				{
-					// By the time we reach this destructor, the first singleton's destructor should have been executed already.
-					check(FThreadSingletonFirst::TryGet() == nullptr);
+					// By the time we reach this destructor, the second singleton's destructor should have been executed already.
+					check(FThreadSingletonSecond::TryGet() == nullptr);
 				}
 			};
 			FThread Thread;
