Fix thread_local usage in concurrent linear allocator that can cause crashes with merged modules

On some platforms, it is not possible to unload a dynamic library if it contains a thread_local static that was accessed by any living thread.
The highly templated nature of allocator code means self-cleanup TLS code in there ends up compiled into libraries for innocuous-looking code.
The fix is to use platform-TLS instead, and the auto-cleanup mechanism already provided.

- Add TModuleSafeThreadLocal that acts as a thread_local object holder
- Add PLATFORM_THREAD_LOCAL_IS_UNSAFE_IN_MODULES defininition for relevant platforms
- Add UE_MODULE_SAFE_THREAD_LOCAL that branches to either TModuleSafeThreadLocal  or thread_local
- Use UE_MODULE_SAFE_THREAD_LOCAL in allocation code

#rnx

#Synced-CL 51619035

diff --git a/Engine/Source/Runtime/Core/Public/Experimental/ConcurrentLinearAllocator.h b/Engine/Source/Runtime/Core/Public/Experimental/ConcurrentLinearAllocator.h
--- a/Engine/Source/Runtime/Core/Public/Experimental/ConcurrentLinearAllocator.h
+++ b/Engine/Source/Runtime/Core/Public/Experimental/ConcurrentLinearAllocator.h
@@ -76,7 +76,7 @@ class TBlockAllocationCache
 
 	inline static void* SwapBlock(void* NewBlock)
 	{
-		static thread_local FTLSCleanup Tls;
+		UE_MODULE_SAFE_THREAD_LOCAL(FTLSCleanup, Tls);
 		void* Ret = Tls.Block;
 		Tls.Block = NewBlock;
 		return Ret;
@@ -324,7 +324,7 @@ class TLinearAllocatorBase
 				ASAN_POISON_MEMORY_REGION( Header + 1, BlockAllocationTag::BlockSize - sizeof(FBlockHeader) );
 			}
 
-			static thread_local FTLSCleanup Cleanup;
+			UE_MODULE_SAFE_THREAD_LOCAL(FTLSCleanup, Cleanup);
 			new (&Cleanup) FTLSCleanup{ Header };
 		}
 	}
@@ -370,7 +370,7 @@ public:
 #endif
 		}
 
-		static thread_local FBlockHeader* Header;
+		UE_MODULE_SAFE_THREAD_LOCAL(FBlockHeader*, Header);
 		if (Header == nullptr)
 		{
 		AllocateNewBlock:
@@ -386,7 +386,7 @@ public:
 					ASAN_POISON_MEMORY_REGION( Header + 1, BlockAllocationTag::BlockSize - sizeof(FBlockHeader) );
 				}
 
-				static thread_local FTLSCleanup Cleanup;
+				UE_MODULE_SAFE_THREAD_LOCAL(FTLSCleanup, Cleanup);
 				new (&Cleanup) FTLSCleanup{ Header };
 			}
 			else
diff --git a/Engine/Source/Runtime/Core/Public/HAL/Platform.h b/Engine/Source/Runtime/Core/Public/HAL/Platform.h
--- a/Engine/Source/Runtime/Core/Public/HAL/Platform.h
+++ b/Engine/Source/Runtime/Core/Public/HAL/Platform.h
@@ -671,6 +671,10 @@
 	#define PLATFORM_BREAK_IS_EXPRESSION 0
 #endif
 
+#ifndef PLATFORM_THREAD_LOCAL_IS_UNSAFE_IN_MODULES
+	#define PLATFORM_THREAD_LOCAL_IS_UNSAFE_IN_MODULES 0
+#endif
+
 // deprecated, do not use
 #define PLATFORM_HAS_THREADSAFE_RHIGetRenderQueryResult	#
 #define PLATFORM_SUPPORTS_RHI_THREAD #
diff --git a/Engine/Source/Runtime/Core/Public/HAL/PlatformTLS.h b/Engine/Source/Runtime/Core/Public/HAL/PlatformTLS.h
--- a/Engine/Source/Runtime/Core/Public/HAL/PlatformTLS.h
+++ b/Engine/Source/Runtime/Core/Public/HAL/PlatformTLS.h
@@ -6,3 +6,87 @@
 #include "GenericPlatform/GenericPlatformTLS.h"
 
 #include COMPILED_PLATFORM_HEADER(PlatformTLS.h) // IWYU pragma: export
+
+#if PLATFORM_THREAD_LOCAL_IS_UNSAFE_IN_MODULES
+
+#include "HAL/TlsAutoCleanup.h"
+#include "Misc/AssertionMacros.h"
+
+namespace UE::Core::Private
+{
+	/**
+	 * TModuleSafeThreadLocal is an alternative to thread_local that can safely be used in unloadable dynamic libraries on all platforms.
+	 * Please see UE_MODULE_SAFE_THREAD_LOCAL for usage and limitations.
+	 */
+	template<class T>
+	class TModuleSafeThreadLocal
+	{
+		static_assert(std::is_default_constructible_v<T>, "Thread-local object must be default-constructible");
+		static_assert(std::is_copy_constructible_v<T>, "Thread-local object must be copy-constructible");
+
+	public:
+		UE_NONCOPYABLE(TModuleSafeThreadLocal);
+
+		/** Constructor, allocates a TLS slot for this instance. */
+		TModuleSafeThreadLocal()
+			: TlsSlot(FPlatformTLS::AllocTlsSlot())
+		{
+			verify(TlsSlot != FPlatformTLS::InvalidTlsSlot);
+		}
+
+		/** Destructor, frees the TLS slot for this instance. */
+		~TModuleSafeThreadLocal()
+		{
+			FPlatformTLS::FreeTlsSlot(TlsSlot);
+		}
+
+		/**
+		 * Gets a thread-local instance of the value held, default-constructed on the first access.
+		 *
+		 * @return A reference to the per-thread instance of the value.
+		 */
+		T& Get()
+		{
+			TTlsAutoCleanupValue<T>* TlsValue = static_cast<TTlsAutoCleanupValue<T>*>(FPlatformTLS::GetTlsValue(TlsSlot));
+
+			if (TlsValue == nullptr)
+			{
+				TlsValue = new TTlsAutoCleanupValue<T>(T());
+				TlsValue->Register();
+				FPlatformTLS::SetTlsValue(TlsSlot, TlsValue);
+			}
+
+			return TlsValue->Get();
+		}
+
+	private:
+
+		uint32 TlsSlot;
+	};
+}
+
+/**
+ * Declares a variable Name of Type that acts thread_local in a way that is safe for use in unloadable dynamic libraries on all platforms.
+ *
+ * While we can and should use thread_local for TLS in most code, dynamic libraries designed to be unloaded cannot hold TLS objects on some platforms.
+ * This is of particular concern when using merged modules on client builds, where thread_local usage in templates can lead to unintended TLS objects in libraries.
+ *
+ * Limitations:
+ *	Usage of this class is subject to existing limits on platform TLS usage, so this should be used sparingly, only for code that has to be used in unloadable DLLs.
+ *	The type used must be default-constructible to a known initial state, and copy-constructible.
+ *
+ * Example:
+ *
+ * 	struct FMyStruct {...};                                   // Will be default-constructed for each thread that uses it
+ * 	UE_MODULE_SAFE_THREAD_LOCAL(FMyStruct, PerThreadMyStruct) // Each thread has its own PerThreadMyStruct instance
+ */
+#define UE_MODULE_SAFE_THREAD_LOCAL(Type, Name) \
+	static UE::Core::Private::TModuleSafeThreadLocal<Type> TlsStorage##Name;\
+	Type& Name = TlsStorage##Name.Get();
+
+#else
+
+// Identical to thread_local on platforms that do not opt into TModuleSafeThreadLocal
+#define UE_MODULE_SAFE_THREAD_LOCAL(Type, Name) static thread_local Type Name;
+
+#endif // PLATFORM_THREAD_LOCAL_IS_UNSAFE_IN_MODULES
diff --git a/Engine/Source/Runtime/Core/Public/HAL/TlsAutoCleanup.h b/Engine/Source/Runtime/Core/Public/HAL/TlsAutoCleanup.h
--- a/Engine/Source/Runtime/Core/Public/HAL/TlsAutoCleanup.h
+++ b/Engine/Source/Runtime/Core/Public/HAL/TlsAutoCleanup.h
@@ -30,12 +30,23 @@ public:
 		: Value(InValue)
 	{ }
 
+	/** Constructor. */
+	TTlsAutoCleanupValue(T&& InValue)
+		: Value(MoveTemp(InValue))
+	{ }
+
 	/** Gets the value. */
 	T Get() const
 	{
 		return Value;
 	}
 
+	/** Gets the value. */
+	T& Get()
+	{
+		return Value;
+	}
+
 	/* Sets the value. */
 	void Set(const T& InValue)
 	{
