UE: Temporarily set PURGATORY_STOMP_MAX_PURGATORY_MEM to 80GB to help track down a memory issue

diff --git a/Engine/Source/Runtime/Core/Private/HAL/UnrealMemory.cpp b/Engine/Source/Runtime/Core/Private/HAL/UnrealMemory.cpp
--- a/Engine/Source/Runtime/Core/Private/HAL/UnrealMemory.cpp
+++ b/Engine/Source/Runtime/Core/Private/HAL/UnrealMemory.cpp
@@ -63,7 +63,7 @@ class FMallocPurgatoryProxy : public FMalloc
 	enum
 	{
 		PURGATORY_STOMP_CHECKS_FRAMES = 4, // we want to allow several frames since it is theoretically possible for something to be blocked mid-allocation for many frames.
-		PURGATORY_STOMP_MAX_PURGATORY_MEM = 100000000, // for loading and whatnot, we really can't keep this stuff forever
+		PURGATORY_STOMP_MAX_PURGATORY_MEM_KB = 83886080, // for loading and whatnot, we really can't keep this stuff forever
 		PURGATORY_STOMP_CHECKS_CANARYBYTE = 0xdc,
 	};
 	uint32 LastCheckFrame; // unitialized; won't matter
@@ -124,7 +124,7 @@ public:
 		uint32 LocalLastCheckFrame = LastCheckFrame;
 		uint32 LocalGFrameNumber = GFrameNumber;
 
-		bool bFlushAnyway = OutstandingSizeInKB.GetValue() > PURGATORY_STOMP_MAX_PURGATORY_MEM / 1024;
+		bool bFlushAnyway = OutstandingSizeInKB.GetValue() > PURGATORY_STOMP_MAX_PURGATORY_MEM_KB;
 
 		if (bFlushAnyway || LocalLastCheckFrame != LocalGFrameNumber)
 		{
