Fixed ReadTrailingOctet consuming non-continuation bytes in CodepointFromUtf8.

diff --git a/Engine/Source/Runtime/Core/Private/GenericPlatform/GenericPlatformString.cpp b/Engine/Source/Runtime/Core/Private/GenericPlatform/GenericPlatformString.cpp
--- a/Engine/Source/Runtime/Core/Private/GenericPlatform/GenericPlatformString.cpp
+++ b/Engine/Source/Runtime/Core/Private/GenericPlatform/GenericPlatformString.cpp
@@ -964,14 +964,15 @@ namespace UE::Core::Private
 		}
 
 		uint32 Octet = (uint32)(uint8)*Ptr;
-		PopFront(Ptr, SourceEnd);
 
-		// Format isn't 10xxxxxx?
+		// Format isn't 10xxxxxx? Do not advance the pointer so the byte can be
+		// re-parsed as the start of the next sequence.
 		if ((Octet & 192) != 128)
 		{
 			return false;
 		}
 
+		PopFront(Ptr, SourceEnd);
 		OutOctet = Octet;
 		return true;
 	};
