Clear expired timers from the timer manager before module unmount

- Add a new event to the module manager that fires right before unmounting merged library, the same way delegate lifetimes are being checked
- Iterate over all timers in the timer manager at this event to cull any unbound event
- This is necessary when destroyed module objects clear an Active timer, which is then left in the ActivePendingRemoval state at unmount, causing a crash when it does get destroyed.

#rnx
#Synced-CL 51619035

#Synced-CL 51706405

diff --git a/Engine/Source/Runtime/Core/Private/Modules/ModuleManager.cpp b/Engine/Source/Runtime/Core/Private/Modules/ModuleManager.cpp
--- a/Engine/Source/Runtime/Core/Private/Modules/ModuleManager.cpp
+++ b/Engine/Source/Runtime/Core/Private/Modules/ModuleManager.cpp
@@ -897,6 +897,9 @@ void FModuleManager::OnObjectCleanup()
 			{
 				UE_LOG(LogModuleManager, Log, TEXT("OnObjectCleanup: unloading dynamic library '%s'"), *LibraryAndModulesIterator.Key());
 
+				// Give various systems a last chance to clean up before removal
+				OnLibraryUnmountCallback.Broadcast();
+
 #if UE_DELEGATE_CHECK_LIFETIME
 				FMulticastDelegateTracker::CheckForStaleDelegatesInModules(LibraryAndModulesIterator.Value());
 #endif // UE_DELEGATE_CHECK_LIFETIME
diff --git a/Engine/Source/Runtime/Core/Public/Modules/ModuleManager.h b/Engine/Source/Runtime/Core/Public/Modules/ModuleManager.h
--- a/Engine/Source/Runtime/Core/Public/Modules/ModuleManager.h
+++ b/Engine/Source/Runtime/Core/Public/Modules/ModuleManager.h
@@ -547,7 +547,8 @@ public:
 	}
 
 	/**
-	 * Gets a multicast delegate that is executed when any UObjects need destroying before modules are unloaded.
+	 * Gets a multicast delegate that is executed to trigger the removal of all UObjects belonging to a list of modules, when these modules are unloaded.
+	 * This event only occurs with merged modules, when a merged library is being unmounted.
 	 *
 	 * @return The delegate.
 	 */
@@ -556,6 +557,17 @@ public:
 		return OnModulesUnloadCallback;
 	}
 
+	/**
+	 * Gets a multicast delegate that is executed right before a merged library is unmounted, after every object in linked modules has been destroyed.
+	 * This event only occurs with merged modules, when a merged library is being unmounted.
+	 *
+	 * @return The delegate.
+	 */
+	TMulticastDelegateRegistration<void()>& OnLibraryUnmount()
+	{
+		return OnLibraryUnmountCallback;
+	}
+
 	/**
 	 * Gets a delegate that is executed whenever we need to confirm that modules can be safely unloaded.
 	 *
@@ -764,6 +776,9 @@ private:
 	/** Multicast delegate called to remove any loaded objects for a list of modules. */
 	TMulticastDelegate<void(TConstArrayView<FName>)> OnModulesUnloadCallback;
 
+	/** Multicast delegate called to remove anything left before a library unmounts. */
+	TMulticastDelegate<void()> OnLibraryUnmountCallback;
+
 	/** Delegate called to check whether any module to delay-unload still has live objects. */
 	TDelegate<bool(TConstArrayView<FName>)> CheckLiveObjectsInModulesCallback;
 
diff --git a/Engine/Source/Runtime/Engine/Private/TimerManager.cpp b/Engine/Source/Runtime/Engine/Private/TimerManager.cpp
--- a/Engine/Source/Runtime/Engine/Private/TimerManager.cpp
+++ b/Engine/Source/Runtime/Engine/Private/TimerManager.cpp
@@ -8,6 +8,7 @@
 #include "AutoRTFM.h"
 #include "Containers/StringConv.h"
 #include "Engine/GameInstance.h"
+#include "Modules/ModuleManager.h"
 #include "Engine/TimerHandle.h"
 #include "Logging/LogScopedCategoryAndVerbosityOverride.h"
 #include "Misc/AssertionMacros.h"
@@ -347,10 +348,14 @@ FTimerManager::FTimerManager(UGameInstance* GameInstance)
 	{
 		SetGameInstance(GameInstance);
 	}
+
+	FModuleManager::Get().OnLibraryUnmount().AddRaw(this, &FTimerManager::RemoveExpiredTimers);
 }
 
 FTimerManager::~FTimerManager()
 {
+	FModuleManager::Get().OnLibraryUnmount().RemoveAll(this);
+
 	if (IsRunningDedicatedServer())
 	{
 		FCoreDelegates::OnHandleSystemError.RemoveAll(this);
@@ -1527,6 +1532,33 @@ bool FTimerManager::WillRemoveTimerAssert(FTimerHandle Handle) const
 	return false;
 }
 
+void FTimerManager::RemoveExpiredTimers()
+{
+	TArray<FTimerHandle> HandlesToClear;
+
+	for (const FTimerData& TimerData : Timers)
+	{
+		if (!TimerData.TimerDelegate.IsBound())
+		{
+			HandlesToClear.Add(TimerData.Handle);
+		}
+	}	
+
+	for (FTimerHandle Handle : HandlesToClear)
+	{
+		int32 IndexIndex = ActiveTimerHeap.Find(Handle);
+		if (IndexIndex != INDEX_NONE)
+		{
+			ActiveTimerHeap.HeapRemoveAt(IndexIndex, FTimerHeapOrder(Timers), EAllowShrinking::No);
+		}
+
+		PendingTimerSet.Remove(Handle);
+		PausedTimerSet.Remove(Handle);
+
+		RemoveTimer(Handle);
+	}
+}
+
 FTimerHandle FTimerManager::GenerateHandle(int32 Index)
 {
 	uint64 NewSerialNumber = ++LastAssignedSerialNumber;
diff --git a/Engine/Source/Runtime/Engine/Public/TimerManager.h b/Engine/Source/Runtime/Engine/Public/TimerManager.h
--- a/Engine/Source/Runtime/Engine/Public/TimerManager.h
+++ b/Engine/Source/Runtime/Engine/Public/TimerManager.h
@@ -494,6 +494,9 @@ private:
 	AUTORTFM_DISABLE ENGINE_API void RemoveTimer(FTimerHandle Handle);
 	ENGINE_API bool WillRemoveTimerAssert(FTimerHandle Handle) const;
 
+	/** Check for expired timers that are pending destruction, and destroy them */
+	ENGINE_API void RemoveExpiredTimers();
+
 #if UE_AUTORTFM
 	/**
 	 * Registers an on-commit and on-abort handler to transactionally commit or
