Ensure we don't drop the end of output from FMonitoredProcess

#rnx
#Synced-CL 51716730

diff --git a/Engine/Source/Runtime/Core/Private/Misc/MonitoredProcess.cpp b/Engine/Source/Runtime/Core/Private/Misc/MonitoredProcess.cpp
--- a/Engine/Source/Runtime/Core/Private/Misc/MonitoredProcess.cpp
+++ b/Engine/Source/Runtime/Core/Private/Misc/MonitoredProcess.cpp
@@ -57,7 +57,11 @@ bool FMonitoredProcess::RunProcessSynchronous( const FString& InURL, const FStri
 	}
 
 	// Add any remaining unprocessed output.
-	OutOutput.Add(MonitoredProcess.GetFullOutputWithoutDelegate());
+	const FString& RemainingBuffer = MonitoredProcess.GetFullOutputWithoutDelegate();
+	if (RemainingBuffer.Len() > 0)
+	{
+		OutOutput.Add(RemainingBuffer);
+	}
 
 	return MonitoredProcess.GetReturnCode() == 0;
 }
