[HAL] Added DelegatedConsoleVariable, with a user-provided Getter/Setter. No persistence or history, thread safety up to user.

#Synced-CL 51660293

diff --git a/Engine/Source/Runtime/Core/Private/HAL/ConsoleManager.cpp b/Engine/Source/Runtime/Core/Private/HAL/ConsoleManager.cpp
--- a/Engine/Source/Runtime/Core/Private/HAL/ConsoleManager.cpp
+++ b/Engine/Source/Runtime/Core/Private/HAL/ConsoleManager.cpp
@@ -1473,7 +1473,128 @@ template<> TConsoleVariableData<FString>* FConsoleVariable<FString>::AsVariableS
 	return &Data;
 }
 
+/**
+ * A thin console variable that delegates all reads to a Getter and all writes to a Setter.
+ * No history, no priority tracking, no platform value caching - the external source is authoritative.
+ * Inherits from FConsoleVariableBase for flags, help text, and change callbacks only.
+ */
+template <class T>
+class FDelegatedConsoleVariable : public FConsoleVariableBase
+{
+public:
+	FDelegatedConsoleVariable(
+		T InDefaultValue,
+		TFunction<T()> InGetter,
+		TFunction<void(const T&)> InSetter,
+		const TCHAR* Help,
+		EConsoleVariableFlags InFlags)
+		: FConsoleVariableBase(Help, InFlags)
+		, DefaultValue(MoveTemp(InDefaultValue))
+		, Getter(MoveTemp(InGetter))
+		, Setter(MoveTemp(InSetter))
+	{
+		check(Getter);
+		check(Setter);
+		check(!(InFlags & ECVF_RenderThreadSafe)); // Not supported - no shadow data
+	}
+
+	virtual void Release() override { delete this; }
+
+	virtual FResolvedContext ResolveContext(const FSetContext& Context) override
+	{
+		// No history system - just pass through the flags and tag as-is
+		return FResolvedContext{ Context.Flags, Context.Tag };
+	}
+
+	virtual void Set(const TCHAR* Value, const FResolvedContext& Context) override
+	{
+		if (CanChange(Context.Flags))
+		{
+			T ParsedValue;
+			UE::ConsoleManager::Private::GetValueFromString<T>(ParsedValue, Value);
+			Setter(ParsedValue);
+			OnChanged(Context.Flags, false);
+		}
+	}
+
+	virtual void Unset(EConsoleVariableFlags SetBy, FName Tag) override
+	{
+		// The getter is always authoritative, Unset has no effect
+	}
+
+	virtual bool GetBool() const override   { return FConsoleVariableConversionHelper<T>::GetBool(Getter()); }
+	virtual int32 GetInt() const override   { return FConsoleVariableConversionHelper<T>::GetInt(Getter()); }
+	virtual float GetFloat() const override { return FConsoleVariableConversionHelper<T>::GetFloat(Getter()); }
+	virtual FString GetString() const override { return FConsoleVariableConversionHelper<T>::GetString(Getter()); }
+
+	virtual FString GetDefaultValue() override
+	{
+		return TTypeToString<T>::ToString(DefaultValue);
+	}
 
+	// History not applicable
+	virtual void LogHistory(FOutputDevice& Ar) override {}
+	virtual SIZE_T GetHistorySize() override { return 0; }
+
+#if ALLOW_OTHER_PLATFORM_CONFIG
+	virtual TSharedPtr<IConsoleVariable> GetPlatformValueVariable(FName PlatformName, const FString& DeviceProfileName) override { return nullptr; }
+	virtual bool HasPlatformValueVariable(FName PlatformName, const FString& DeviceProfileName) override { return false; }
+	virtual void SetOtherPlatformValue(const TCHAR* InValue, EConsoleVariableFlags SetBy, FName Tag) override {}
+#endif
+
+private:
+	T DefaultValue;
+	TFunction<T()> Getter;
+	TFunction<void(const T&)> Setter;
+};
+
+template<>
+void FDelegatedConsoleVariable<FString>::Set(
+	const TCHAR* Value,
+	const FResolvedContext& Context)
+{
+	if (CanChange(Context.Flags))
+	{
+		Setter(FString(Value));
+		OnChanged(Context.Flags, false);
+	}
+}
+
+template <class T>
+IConsoleVariable* FConsoleManager::RegisterDelegatedConsoleVariableImpl(
+	const TCHAR* Name, T DefaultValue,
+	TFunction<T()> Getter, TFunction<void(const T&)> Setter,
+	const TCHAR* Help, uint32 Flags)
+{
+	return AddConsoleObject(
+		Name,
+		new FDelegatedConsoleVariable<T>(
+			MoveTemp(DefaultValue),
+			MoveTemp(Getter),
+			MoveTemp(Setter),
+			Help,
+			(EConsoleVariableFlags)Flags))->AsVariable();
+}
+
+IConsoleVariable* FConsoleManager::RegisterDelegatedConsoleVariable(const TCHAR* Name, int32 DefaultValue, TFunction<int32()> Getter, TFunction<void(const int32&)> Setter, const TCHAR* Help, uint32 Flags)
+{
+	return RegisterDelegatedConsoleVariableImpl(Name, DefaultValue, MoveTemp(Getter), MoveTemp(Setter), Help, Flags);
+}
+
+IConsoleVariable* FConsoleManager::RegisterDelegatedConsoleVariable(const TCHAR* Name, float DefaultValue, TFunction<float()> Getter, TFunction<void(const float&)> Setter, const TCHAR* Help, uint32 Flags)
+{
+	return RegisterDelegatedConsoleVariableImpl(Name, DefaultValue, MoveTemp(Getter), MoveTemp(Setter), Help, Flags);
+}
+
+IConsoleVariable* FConsoleManager::RegisterDelegatedConsoleVariable(const TCHAR* Name, bool DefaultValue, TFunction<bool()> Getter, TFunction<void(const bool&)> Setter, const TCHAR* Help, uint32 Flags)
+{
+	return RegisterDelegatedConsoleVariableImpl(Name, DefaultValue, MoveTemp(Getter), MoveTemp(Setter), Help, Flags);
+}
+
+IConsoleVariable* FConsoleManager::RegisterDelegatedConsoleVariable(const TCHAR* Name, const FString& DefaultValue, TFunction<FString()> Getter, TFunction<void(const FString&)> Setter, const TCHAR* Help, uint32 Flags)
+{
+	return RegisterDelegatedConsoleVariableImpl(Name, DefaultValue, MoveTemp(Getter), MoveTemp(Setter), Help, Flags);
+}
 
 template<class T>
 FConsoleVariableExtendedData<T>::FConsoleVariableExtendedData(const T& DefaultValue, const TCHAR* Help, EConsoleVariableFlags Flags)
diff --git a/Engine/Source/Runtime/Core/Private/HAL/ConsoleManager.h b/Engine/Source/Runtime/Core/Private/HAL/ConsoleManager.h
--- a/Engine/Source/Runtime/Core/Private/HAL/ConsoleManager.h
+++ b/Engine/Source/Runtime/Core/Private/HAL/ConsoleManager.h
@@ -17,6 +17,11 @@ public:
 	virtual IConsoleVariable* RegisterConsoleVariable(const TCHAR* Name, const TCHAR* DefaultValue, const TCHAR* Help, uint32 Flags) override;
 	virtual IConsoleVariable* RegisterConsoleVariable(const TCHAR* Name, const FString& DefaultValue, const TCHAR* Help, uint32 Flags) override;
 
+	virtual IConsoleVariable* RegisterDelegatedConsoleVariable(const TCHAR* Name, int32 DefaultValue, TFunction<int32()> Getter, TFunction<void(const int32&)> Setter, const TCHAR* Help, uint32 Flags) override;
+	virtual IConsoleVariable* RegisterDelegatedConsoleVariable(const TCHAR* Name, float DefaultValue, TFunction<float()> Getter, TFunction<void(const float&)> Setter, const TCHAR* Help, uint32 Flags) override;
+	virtual IConsoleVariable* RegisterDelegatedConsoleVariable(const TCHAR* Name, bool DefaultValue, TFunction<bool()> Getter, TFunction<void(const bool&)> Setter, const TCHAR* Help, uint32 Flags) override;
+	virtual IConsoleVariable* RegisterDelegatedConsoleVariable(const TCHAR* Name, const FString& DefaultValue, TFunction<FString()> Getter, TFunction<void(const FString&)> Setter, const TCHAR* Help, uint32 Flags) override;
+
 	virtual IConsoleVariable* RegisterConsoleVariableRef(const TCHAR* Name, bool& RefValue, const TCHAR* Help, uint32 Flags) override;
 	virtual IConsoleVariable* RegisterConsoleVariableRef(const TCHAR* Name, int32& RefValue, const TCHAR* Help, uint32 Flags) override;
 	virtual IConsoleVariable* RegisterConsoleVariableRef(const TCHAR* Name, float& RefValue, const TCHAR* Help, uint32 Flags) override;
@@ -89,6 +94,12 @@ public:
 
 private:
 
+	template <class T>
+	IConsoleVariable* RegisterDelegatedConsoleVariableImpl(
+		const TCHAR* Name, T DefaultValue,
+		TFunction<T()> Getter, TFunction<void(const T&)> Setter,
+		const TCHAR* Help, uint32 Flags);
+
 	/**
 	 * @param Stream must not be 0
 	 * @param Pattern must not be 0
diff --git a/Engine/Source/Runtime/Core/Public/HAL/IConsoleManager.h b/Engine/Source/Runtime/Core/Public/HAL/IConsoleManager.h
--- a/Engine/Source/Runtime/Core/Public/HAL/IConsoleManager.h
+++ b/Engine/Source/Runtime/Core/Public/HAL/IConsoleManager.h
@@ -1022,6 +1022,39 @@ struct IConsoleManager
 	* @param Flags bitmask combined from EConsoleVariableFlags
 	*/
 	virtual IConsoleVariable* RegisterConsoleVariableRef(const TCHAR* Name, FName& RefValue, const TCHAR* Help, uint32 Flags = ECVF_Default) = 0;
+
+	/**
+	* Create a console variable that delegates get/set to external functions.
+	* The Getter is called on every read; the Setter is called on every write.
+	* ECVF_RenderThreadSafe is NOT supported with delegate variables.
+	*
+	* @param Name must not be 0, must not be empty, must not contain spaces
+	* @param DefaultValue initial value passed to Setter during registration
+	* @param Getter called to read the current value
+	* @param Setter called to write a new value
+	* @param Help must not be 0, must not be empty
+	* @param Flags bitmask combined from EConsoleVariableFlags
+	*/
+	virtual IConsoleVariable* RegisterDelegatedConsoleVariable(
+		const TCHAR* Name, int32 DefaultValue,
+		TFunction<int32()> Getter, TFunction<void(const int32&)> Setter,
+		const TCHAR* Help, uint32 Flags = ECVF_Default) = 0;
+
+	virtual IConsoleVariable* RegisterDelegatedConsoleVariable(
+		const TCHAR* Name, float DefaultValue,
+		TFunction<float()> Getter, TFunction<void(const float&)> Setter,
+		const TCHAR* Help, uint32 Flags = ECVF_Default) = 0;
+
+	virtual IConsoleVariable* RegisterDelegatedConsoleVariable(
+		const TCHAR* Name, bool DefaultValue,
+		TFunction<bool()> Getter, TFunction<void(const bool&)> Setter,
+		const TCHAR* Help, uint32 Flags = ECVF_Default) = 0;
+
+	virtual IConsoleVariable* RegisterDelegatedConsoleVariable(
+		const TCHAR* Name, const FString& DefaultValue,
+		TFunction<FString()> Getter, TFunction<void(const FString&)> Setter,
+		const TCHAR* Help, uint32 Flags = ECVF_Default) = 0;
+
 	/**
 	 * Create a reference to a show flag variable
 	 * @param CVarName must not be 0, e.g. "Show.PostProcessing"
@@ -1918,6 +1951,27 @@ public:
 		AsVariable()->SetOnChangedCallback(Callback);
 	}
 
+	/**
+	 * Create a console variable that delegates get/set to external functions.
+	 * The Getter is called on every read; the Setter is called on every write.
+	 * ECVF_RenderThreadSafe is NOT supported with delegate variables.
+	 *
+	 * @param Name must not be 0
+	 * @param DefaultValue initial value; the Getter is authoritative after construction
+	 * @param Getter called to read the current value
+	 * @param Setter called to write a new value
+	 * @param Help must not be 0
+	 * @param Flags bitmask combined from EConsoleVariableFlags
+	 */
+	FORCENOINLINE TAutoConsoleVariable(
+		const TCHAR* Name, const T& DefaultValue,
+		TFunction<T()> Getter, TFunction<void(const T&)> Setter,
+		const TCHAR* Help, uint32 Flags = ECVF_Default)
+		: TAutoConsoleVariable(IConsoleManager::Get().RegisterDelegatedConsoleVariable(
+			Name, DefaultValue, MoveTemp(Getter), MoveTemp(Setter), Help, Flags))
+	{
+	}
+
 	virtual ~TAutoConsoleVariable()
 	{
 		AsVariable()->ClearOnChangedCallback();
@@ -1925,17 +1979,35 @@ public:
 
 	T GetValueOnGameThread() const
 	{
-		return Ref->GetValueOnGameThread();
+		if (Ref)
+		{
+			return Ref->GetValueOnGameThread();
+		}
+		T Result;
+		AsVariable()->GetValue(Result);
+		return Result;
 	}
 
 	T GetValueOnRenderThread() const
 	{
-		return Ref->GetValueOnRenderThread();
+		if (Ref)
+		{
+			return Ref->GetValueOnRenderThread();
+		}
+		T Result;
+		AsVariable()->GetValue(Result);
+		return Result;
 	}
 
 	T GetValueOnAnyThread(bool bForceGameThread = false) const
 	{
-		return Ref->GetValueOnAnyThread(bForceGameThread);
+		if (Ref)
+		{
+			return Ref->GetValueOnAnyThread(bForceGameThread);
+		}
+		T Result;
+		AsVariable()->GetValue(Result);
+		return Result;
 	}
 	
 	/** Dereference back to a variable**/
@@ -1957,7 +2029,7 @@ public:
 		return AsVariable();
 	}
 private:
-	TConsoleVariableData<T>* Ref;
+	TConsoleVariableData<T>* Ref = nullptr;
 };
 
 #else
