[core] xxhash64 additions - cb read/write, to/from string helpers, type layout intrinsics for memory image serialization

#Synced-CL 51673377

diff --git a/Engine/Source/Runtime/Core/Private/Hash/xxhash.cpp b/Engine/Source/Runtime/Core/Private/Hash/xxhash.cpp
--- a/Engine/Source/Runtime/Core/Private/Hash/xxhash.cpp
+++ b/Engine/Source/Runtime/Core/Private/Hash/xxhash.cpp
@@ -5,9 +5,12 @@
 #include "Async/ParallelFor.h"
 #include "Containers/ContainersFwd.h"
 #include "Containers/StaticArray.h"
+#include "Containers/StringView.h"
+#include "Containers/UnrealString.h"
 #include "Memory/CompositeBuffer.h"
 #include "Memory/MemoryView.h"
 #include "Memory/SharedBuffer.h"
+#include "Serialization/CompactBinaryWriter.h"
 
 THIRD_PARTY_INCLUDES_START
 #define XXH_INLINE_ALL
@@ -67,6 +70,38 @@ FXxHash64 FXxHash64::HashBufferChunked(const void* Data, uint64 Size, uint64 Chu
 	return PrivateHashBufferChunked<FXxHash64>(FMemoryView(Data, Size), ChunkSize);
 }
 
+FString FXxHash64::ToString() const
+{
+	FString Out;
+	AppendString(Out);
+	return Out;
+}
+
+void FXxHash64::FromString(FStringView Src)
+{
+	check(Src.Len() == GetStringLen());
+	uint8 Bytes[sizeof(Hash)];
+	UE::String::HexToBytes(Src, Bytes);
+	*this = FromByteArray(Bytes);
+}
+
+void FXxHash64::AppendString(FString& Out) const
+{
+	uint8 Bytes[sizeof(Hash)];
+	ToByteArray(Bytes);
+	BytesToHex(Bytes, sizeof(Bytes), Out);
+}
+
+FString LexToString(const FXxHash64& InHash)
+{
+	return InHash.ToString();
+}
+
+void LexFromString(FXxHash64& InHash, const TCHAR* InString)
+{
+	InHash.FromString(InString);
+}
+
 FXxHash128 FXxHash128::HashBuffer(FMemoryView View)
 {
 	const XXH128_hash_t Value = XXH3_128bits(View.GetData(), View.GetSize());
@@ -89,6 +124,28 @@ FXxHash128 FXxHash128::HashBufferChunked(const void* Data, uint64 Size, uint64 C
 	return PrivateHashBufferChunked<FXxHash128>(FMemoryView(Data, Size), ChunkSize);
 }
 
+bool FXxHash64::LoadFromCompactBinary(const FCbFieldView& Field)
+{
+	FMemoryView BytesView = const_cast<FCbFieldView&>(Field).AsBinaryView();
+	if (BytesView.GetSize() != sizeof(Hash))
+	{
+		*this = FXxHash64();
+		return false;
+	}
+	const uint8 (*Bytes)[sizeof(Hash)] = reinterpret_cast<const uint8(*)[sizeof(Hash)]>(BytesView.GetData());
+	*this = FromByteArray(*Bytes);
+	return true;
+}
+
+FCbWriter& FXxHash64::WriteCompactBinary(FCbWriter& Writer) const
+{
+	uint8 Bytes[sizeof(Hash)];
+	ToByteArray(Bytes);
+
+	Writer.AddBinary(FMemoryView(Bytes, sizeof(Bytes)));
+	return Writer;
+}
+
 FXxHash128 FXxHash128::HashBuffer(const FCompositeBuffer& Buffer)
 {
 	FXxHash128Builder Builder;
diff --git a/Engine/Source/Runtime/Core/Public/Hash/xxhash.h b/Engine/Source/Runtime/Core/Public/Hash/xxhash.h
--- a/Engine/Source/Runtime/Core/Public/Hash/xxhash.h
+++ b/Engine/Source/Runtime/Core/Public/Hash/xxhash.h
@@ -12,6 +12,8 @@
 #include "String/BytesToHex.h"
 
 class FCompositeBuffer;
+class FCbFieldView;
+class FCbWriter;
 template <typename CharType> class TStringBuilderBase;
 
 /** A 64-bit hash from XXH3. */
@@ -38,7 +40,7 @@ struct FXxHash64
 	[[nodiscard]] CORE_API static FXxHash64 HashBufferChunked(const void* Data, uint64 Size, uint64 ChunkSize);
 
 	/** Load the hash from its canonical (big-endian) representation. */
-	static inline FXxHash64 FromByteArray(const uint8 (&Bytes)[sizeof(uint64)])
+	[[nodiscard]] static inline FXxHash64 FromByteArray(const uint8(&Bytes)[sizeof(uint64)])
 	{
 		uint64 HashBigEndian;
 		FMemory::Memcpy(&HashBigEndian, Bytes, sizeof(uint64));
@@ -72,6 +74,17 @@ struct FXxHash64
 		return Hash == 0;
 	}
 
+	static constexpr int32 GetStringLen() { return sizeof(uint64) * 2; }
+
+	/** Convert the hash to a 16-character hex string. */
+	[[nodiscard]] CORE_API FString ToString() const;
+
+	/** Append the 16-character hex representation to the given string. */
+	CORE_API void AppendString(FString& Out) const;
+
+	/** Parse the hash from a 16-character hex string. */
+	CORE_API void FromString(FStringView Src);
+
 	friend inline uint32 GetTypeHash(const FXxHash64& InHash)
 	{
 		return uint32(InHash.Hash);
@@ -90,8 +103,19 @@ struct FXxHash64
 		UE::String::BytesToHexLower(MakeArrayView(Bytes), Builder);
 		return Builder;
 	}
+
+	CORE_API friend FString LexToString(const FXxHash64& InHash); 
+	CORE_API friend void LexFromString(FXxHash64& InHash, const TCHAR* InString);
+
+	friend FCbWriter& operator<<(FCbWriter& Writer, const FXxHash64& InHash) { return InHash.WriteCompactBinary(Writer); }
+	friend bool LoadFromCompactBinary(const FCbFieldView& Field, FXxHash64& OutHash) { return OutHash.LoadFromCompactBinary(Field); }
+
+	CORE_API bool LoadFromCompactBinary(const FCbFieldView& Field);
+	CORE_API FCbWriter& WriteCompactBinary(FCbWriter& Writer) const;
 };
 
+DECLARE_INTRINSIC_TYPE_LAYOUT(FXxHash64);
+
 /** A 128-bit hash from XXH128. */
 struct FXxHash128
 {
@@ -122,7 +146,7 @@ struct FXxHash128
 	[[nodiscard]] CORE_API static FXxHash128 HashBufferChunked(const void* Data, uint64 Size, uint64 ChunkSize);
 
 	/** Load the hash from its canonical (big-endian) representation. */
-	static inline FXxHash128 FromByteArray(const uint8 (&Bytes)[sizeof(uint64[2])])
+	[[nodiscard]] static inline FXxHash128 FromByteArray(const uint8 (&Bytes)[sizeof(uint64[2])])
 	{
 		uint64 HashBigEndian[2];
 		FMemory::Memcpy(&HashBigEndian, Bytes, sizeof(uint64[2]));
