Clarify that FApp::IsGame checks the current state and not the launched state, this is not a behavior change.
Cleanup other App.h functions related to launch parameters.

#Synced-CL 51720594

diff --git a/Engine/Source/Runtime/Core/Public/Misc/App.h b/Engine/Source/Runtime/Core/Public/Misc/App.h
--- a/Engine/Source/Runtime/Core/Public/Misc/App.h
+++ b/Engine/Source/Runtime/Core/Public/Misc/App.h
@@ -165,18 +165,18 @@ public:
 	}
 
 	/**
-	 * Checks whether this application is a game.
+	 * Checks whether this application is currently running as a game or other packaged application, and is not acting like an editor.
+	 * In an editor build, this changes based on what the game thread is doing for Play In Editor so cannot be reliably accessed from worker threads.
 	 *
-	 * Returns true if a normal or PIE game is active (basically !GIsEdit or || GIsPlayInEditorWorld)
-	 * This must NOT be accessed on threads other than the game thread!
-	 * Use View->Family->EnginShowFlags.Game on the rendering thread.
+	 * Returns true if a normal or PIE game is active, including a dedicated server.
+	 * Use IsRunningGame if you want to check what the process was launched as and not the current state.
 	 *
-	 * @return true if the application is a game, false otherwise.
+	 * @return true if the application is currently acting like a game, false otherwise.
 	 */
 	FORCEINLINE static bool IsGame()
 	{
 #if WITH_EDITOR
-		return !GIsEditor || GIsPlayInEditorWorld || IsRunningGame();
+		return !GIsEditor || GIsPlayInEditorWorld;
 #else
 		return true;
 #endif
@@ -205,8 +205,6 @@ public:
 		GInternalProjectName[UE_ARRAY_COUNT(GInternalProjectName) - 1] = TEXT('\0');
 	}
 
-public:
-
 	/**
 	 * Add the specified user to the list of authorized session users.
 	 *
@@ -347,10 +345,10 @@ public:
 	}
 
 	/**
-	 * Checks whether this is a standalone application.
+	 * Checks whether this is a standalone application, where no session id is specified on the commandline.
 	 *
-	 * An application is standalone if it has not been launched as part of a session from UFE.
-	 * If an application was launched from UFE, it is part of a session and not considered standalone.
+	 * An application is standalone if it has not been launched as part of a session from UFE or another tool.
+	 * This is not related to the networking or Play in Editor settings of the same name.
 	 *
 	 * @return true if this application is standalone, false otherwise.
 	 */
@@ -392,16 +390,6 @@ public:
 		SessionOwner = NewOwner;
 	}
 
-public:
-
-// MSVC 16.4 (1924) has a bug that does not properly handle the local static bool in CanEverRender. This will be fixed in 16.5. Working around by using FORCENOINLINE.
-// MSVC 16.6 (1926) fixes the problem so inlning can be enabled again from that point onwards.
-#if !UE_SERVER && defined(_MSC_VER) && _MSC_VER >= 1924 && _MSC_VER <= 1925
-#define INLINE_CANEVERRENDER FORCENOINLINE
-#else
-#define INLINE_CANEVERRENDER FORCEINLINE
-#endif
-
 	/**
 	 * Checks whether this application can render anything.
 	 * Certain application types never render, while for others this behavior may be controlled by switching to NullRHI.
@@ -409,7 +397,7 @@ public:
 	 *
 	 * @return true if the application can render, false otherwise.
 	 */
-	INLINE_CANEVERRENDER static bool CanEverRender()
+	FORCEINLINE static bool CanEverRender()
 	{
 #if UE_SERVER
 		return false;
@@ -428,7 +416,7 @@ public:
 	 *
 	 * @return true if the application can render, false otherwise.
 	 */
-	INLINE_CANEVERRENDER static bool CanEverRenderOrProduceRenderData()
+	FORCEINLINE static bool CanEverRenderOrProduceRenderData()
 	{
 		return !FPlatformProperties::RequiresCookedData() || FApp::CanEverRender();
 	}
@@ -440,7 +428,7 @@ public:
 	 *
 	 * @return true if the application can render audio, false otherwise.
 	 */
-	INLINE_CANEVERRENDER static bool CanEverRenderAudio()
+	FORCEINLINE static bool CanEverRenderAudio()
 	{
 #if UE_SERVER
 		return false;
@@ -462,7 +450,7 @@ public:
 	 * 
 	 * @return true if the application will mute audio, false otherwise.
 	 */
-	INLINE_CANEVERRENDER static bool IsAudioMuted()
+	FORCEINLINE static bool IsAudioMuted()
 	{
 		static bool bMuteAudio = FParse::Param(FCommandLine::Get(), TEXT("muteaudio"));
 		return bMuteAudio;
