LinearAllocator: Merge the thread-performance and efficient oversize-block features of AssetRegistry's Internal FLinearBlockAllocator into the FLinearBlockAllocator in Memory\LinearAllocator.
Make FLinearBlockAllocator available in all configurations, rename it to be in UE namespace, and make it support allocation from an InnerMalloc parameter.
We plan to use LinearBlockAllocator for MallocProxies that take over GMalloc and need to allocate from their InnerMalloc rather than FMemory or Operating system.
Add edge-case fixes from static analysis elsewhere in LinearAllocator.cpp and AssetRegistryState.cpp.
#rnx
#Synced-CL 51748961

diff --git a/Engine/Source/Runtime/AssetRegistry/Internal/AssetRegistry/LinearBlockAllocator.h b/Engine/Source/Runtime/AssetRegistry/Internal/AssetRegistry/LinearBlockAllocator.h
deleted file mode 100644
diff --git a/Engine/Source/Runtime/AssetRegistry/Private/AssetRegistryState.cpp b/Engine/Source/Runtime/AssetRegistry/Private/AssetRegistryState.cpp
--- a/Engine/Source/Runtime/AssetRegistry/Private/AssetRegistryState.cpp
+++ b/Engine/Source/Runtime/AssetRegistry/Private/AssetRegistryState.cpp
@@ -3,7 +3,6 @@
 #include "AssetRegistry/AssetRegistryState.h"
 
 #include "Algo/Compare.h"
-#include "AssetRegistry/LinearBlockAllocator.h"
 #include "AssetRegistry/ParallelSort.h"
 #include "Algo/Sort.h"
 #include "Algo/Unique.h"
@@ -18,6 +17,7 @@
 #include "DependsNode.h"
 #include "GenericPlatform/GenericPlatformChunkInstall.h"
 #include "Logging/StructuredLog.h"
+#include "Memory/LinearAllocator.h"
 #include "Misc/CommandLine.h"
 #include "Misc/ConfigCacheIni.h"
 #include "Misc/FileHelper.h"
@@ -1431,7 +1431,7 @@ void FilterAssetsByPackageName(
 		if (TOptional<TConstArrayView<FAssetDataPtrIndex>> Array = AccelerationMap.Find(Key))
 		{
 			Matches.Add(*Array);
-			TotalMatches += Matches.Num();
+			TotalMatches += Array->Num();
 		}
 	}
 
@@ -2353,7 +2353,7 @@ TArray<FAssetData*> FAssetRegistryState::GetMutableAssetsSortedByObjectPath()
 			SortDatas.Emplace_GetRef().AssetData = &AssetData;
 		});
 
-	UE::FLinearBlockAllocator SharedAllocator;
+	UE::FLinearBlockAllocator SharedAllocator(nullptr, 0U, ESPMode::NotThreadSafe);
 	TArray<UE::FLinearBlockAllocatorThreadAccessor> ThreadAllocators;
 	// Set minbatchsize to avoid wasting memory by having threads allocate a 64K block and then not use most of it
 	constexpr int32 MinSetupBatchSize = 500;
@@ -2363,6 +2363,7 @@ TArray<FAssetData*> FAssetRegistryState::GetMutableAssetsSortedByObjectPath()
 		// ContextConstructor
 		[&SharedAllocator](int32 ThreadIndex, int32 NumThreads) -> UE::FLinearBlockAllocator&
 		{
+			// ParallelFor passes return value into the FLinearBlockAllocatorThreadAccessor constructor.
 			return SharedAllocator;
 		},
 		// Body
@@ -3718,7 +3719,7 @@ void FAssetRegistryState::RemoveAssetData(FAssetData* AssetData, bool bRemoveDep
 	TBitArray<> bOutRemovedAssetDatas(false, 1);
 	TBitArray<> bOutRemovedPackageDatas(false, 1);
 	RemoveAssetDatas({ &AssetData, 1 }, bRemoveDependencyData, bOutRemovedAssetDatas, bOutRemovedPackageDatas);
-	bRemoveDependencyData = bOutRemovedAssetDatas[0];
+	bOutRemovedAssetData = bOutRemovedAssetDatas[0];
 	bOutRemovedPackageData = bOutRemovedPackageDatas[0];
 }
 
diff --git a/Engine/Source/Runtime/Core/Private/Memory/LinearAllocator.cpp b/Engine/Source/Runtime/Core/Private/Memory/LinearAllocator.cpp
--- a/Engine/Source/Runtime/Core/Private/Memory/LinearAllocator.cpp
+++ b/Engine/Source/Runtime/Core/Private/Memory/LinearAllocator.cpp
@@ -1,18 +1,28 @@
 // Copyright Epic Games, Inc. All Rights Reserved.
 
 #include "Memory/LinearAllocator.h"
+
 #include "BuildSettings.h"
+#include "CoreGlobals.h"
+#include "HAL/LowLevelMemTracker.h"
+#include "Misc/AssertionMacros.h"
+#include "Misc/ScopeExit.h"
+#include "Misc/ScopeLock.h"
+#include "Templates/AlignmentTemplates.h"
 
 CORE_API FPersistentLinearAllocatorExtends GPersistentLinearAllocatorExtends;
 
-static constexpr SIZE_T LinearAllocatorBlockSize = 64 * 1024;
+namespace UE::LinearAllocator
+{
 
+static constexpr uint32 DefaultBlockSize = 64 * 1024;
+static constexpr uint32 BlockAlignment = 8;
+static constexpr uint32 DefaultAlignment = 8;
 
-#if UE_ENABLE_LINEAR_VIRTUAL_ALLOCATOR
+} // namespace UE::LinearAllocator
 
-#include "CoreGlobals.h"
-#include "Misc/ScopeLock.h"
-#include "HAL/LowLevelMemTracker.h"
+
+#if UE_ENABLE_LINEAR_VIRTUAL_ALLOCATOR
 
 struct FPersistentLinearAllocator : public FLinearAllocator
 {
@@ -49,7 +59,7 @@ FLinearAllocator::FLinearAllocator(SIZE_T ReserveMemorySize)
 
 void* FLinearAllocator::Allocate(SIZE_T Size, uint32 Alignment)
 {
-	Alignment = FMath::Max(Alignment, 8u);
+	Alignment = FMath::Max(Alignment, UE::LinearAllocator::DefaultAlignment);
 	{
 		void* Mem = nullptr;
 		{
@@ -61,7 +71,8 @@ void* FLinearAllocator::Allocate(SIZE_T Size, uint32 Alignment)
 				if (NewOffset > Committed)
 				{
 					LLM_PLATFORM_SCOPE(ELLMTag::FMalloc);
-					const SIZE_T ToCommit = Align(NewOffset - Committed, FMath::Max((SIZE_T)VirtualMemory.GetCommitAlignment(), LinearAllocatorBlockSize));
+					const SIZE_T ToCommit = Align(NewOffset - Committed, FMath::Max((SIZE_T)VirtualMemory.GetCommitAlignment(),
+						(SIZE_T)UE::LinearAllocator::DefaultBlockSize));
 					VirtualMemory.Commit(Committed, ToCommit);
 					Committed += ToCommit;
 					LLM_IF_ENABLED(FLowLevelMemTracker::Get().OnLowLevelAlloc(ELLMTracker::Platform, (uint8*)VirtualMemory.GetVirtualPointer() + CurrentOffset, ToCommit));
@@ -88,7 +99,7 @@ void* FLinearAllocator::Allocate(SIZE_T Size, uint32 Alignment)
 
 bool FLinearAllocator::CanFit(SIZE_T Size, uint32 Alignment) const
 {
-	return (Reserved - Align(CurrentOffset, Alignment)) >= Size;
+	return Reserved >= Size + Align(CurrentOffset, Alignment);
 }
 
 bool FLinearAllocator::ContainsPointer(const void* Ptr) const
@@ -98,67 +109,280 @@ bool FLinearAllocator::ContainsPointer(const void* Ptr) const
 
 #else
 
-FLinearBlockAllocator::FLinearBlockAllocator(SIZE_T)
+FLinearAllocator::FLinearAllocator(SIZE_T /* ReserveMemorySize */)
+	: UE::FLinearBlockAllocator(nullptr, 0U, ESPMode::ThreadSafe)
+{
+	// When !UE_ENABLE_LINEAR_VIRTUAL_ALLOCATOR , this allocator is a fallback from a higher level virtual allocator.
+	// We do not need to support the initial reservation in that case. 
+}
+
+typedef FLinearAllocator FPersistentLinearAllocator;
+#endif
+
+namespace UE
+{
+
+FLinearBlockAllocator::FBlock* FLinearBlockAllocator::FBlock::GetHeaderFromBlockStart(void* Block, uint32 BlockSize)
+{
+	return reinterpret_cast<FBlock*>(reinterpret_cast<SIZE_T>(Block) + BlockSize - sizeof(FBlock));
+}
+
+void* FLinearBlockAllocator::FBlock::GetBlockStart()
+{
+	check(((SIZE_T)this) + sizeof(FBlock) >= (SIZE_T)BlockSize);
+	return (void*)(((SIZE_T)this) + sizeof(FBlock) - (SIZE_T)BlockSize);
+}
+
+FLinearBlockAllocator::FLinearBlockAllocator(FMalloc* InInnerMalloc, uint32 InBlockSize,
+	ESPMode InThreadMode)
+	: InnerMalloc(InInnerMalloc ? InInnerMalloc : UE::Private::GMalloc)
+	, DefaultBlockSize(InBlockSize != 0 ? InBlockSize : UE::LinearAllocator::DefaultBlockSize)
+	, ThreadMode(InThreadMode)
 {
-	AllocateNewBlock(LinearAllocatorBlockSize);
-	FirstHeader = CurrentHeader;
+	static_assert(IsAligned(MinBlockSize, UE::LinearAllocator::BlockAlignment));
+	DefaultBlockSize = FMath::Max(DefaultBlockSize, MinBlockSize);
+	DefaultBlockSize = FMath::RoundUpToPowerOfTwo(DefaultBlockSize);
 }
 
 FLinearBlockAllocator::~FLinearBlockAllocator()
 {
+	checkf(ReferenceCount.load(std::memory_order_relaxed) == 0,
+		TEXT("FLinearBlockAllocator is destroyed while still in use by a FLinearBlockAllocatorThreadAccessor."));
+
+	LLM(const bool bDeclareAllocs = (InnerMalloc == UE::Private::GMalloc) && FLowLevelMemTracker::IsEnabled());
+
 	while (FirstHeader)
 	{
-		FBlock* Next = FirstHeader->Next;
-		CurrentBlock = (uint8*)FirstHeader + sizeof(FBlock) - CurrentHeader->LinearAllocatorBlockSize;
-		FMemory::Free(CurrentBlock);
-		FirstHeader = Next;
+		FBlock* NextHeader = FirstHeader->Next;
+		void* FirstBlock = FirstHeader->GetBlockStart();
+		LLM(if (bDeclareAllocs)
+		{
+			FLowLevelMemTracker::Get().OnLowLevelFree(ELLMTracker::Default, FirstBlock, ELLMAllocType::FMalloc);
+		});
+		InnerMalloc->Free(FirstBlock);
+
+		FirstHeader = NextHeader;
+	}
+}
+
+void* FLinearBlockAllocator::Malloc(SIZE_T Size, uint32 Alignment)
+{
+	checkf(Size <= MAX_uint32 - sizeof(FBlock)*2,
+		TEXT("FLinearBlockAllocator does not support allocations of size ~4GB or more. We store BlockSizes as uint32."));
+	checkf(ReferenceCount.load(std::memory_order_relaxed) == 0,
+		TEXT("FLinearBlockAllocator::Malloc is being used directly while in use by a FLinearBlockAllocatorThreadAccessor, this is not allowed, all malloc calls must go through a FLinearBlockAllocatorThreadAccessor."));
+	Alignment = FMath::Max(Alignment, UE::LinearAllocator::DefaultAlignment);
+
+	const bool bThreadSafe = IsThreadSafe();
+	if (bThreadSafe)
+	{
+		Lock.Lock();
 	}
+	ON_SCOPE_EXIT
+	{
+		if (bThreadSafe)
+		{
+			Lock.Unlock();
+		}
+	};
+
+	void* Result = TryAllocateFromBlock(Size, Alignment, LastHeader);
+	if (Result)
+	{
+		return Result;
+	}
+
+	FBlock* NewHeader;
+	if (!RequiresCustomAllocation(Size, Alignment))
+	{
+		NewHeader = AddDefaultBlock();
+	}
+	else
+	{
+		NewHeader = AddCustomBlock(Size, Alignment);
+	}
+	Result = TryAllocateFromBlock(Size, Alignment, NewHeader);
+	check(Result);
+	return Result;
+}
+
+SIZE_T FLinearBlockAllocator::GetAllocatedMemorySize() const
+{
+	FScopeLock ScopeLock(&Lock);
+	return TotalAllocated;
+}
+
+bool FLinearBlockAllocator::IsThreadSafe() const
+{
+	return ThreadMode == ESPMode::ThreadSafe;
 }
 
-void* FLinearBlockAllocator::Allocate(SIZE_T Size, uint32 Alignment)
+void* FLinearBlockAllocator::TryAllocateFromBlock(SIZE_T Size, uint32 Alignment, FBlock* Header)
 {
-	FScopeLock AutoLock(&Lock);
-	if (!CanFit(Size, Alignment))
+	if (!Header)
 	{
-		//TODO: if Size >= BlockSize, allocate a new block and link it between curent and previous block
-		AllocateNewBlock(Size + sizeof(FBlock));
+		return nullptr;
 	}
 
-	CurrentOffset = Align(CurrentOffset, Alignment);
-	uint8* Mem = &CurrentBlock[CurrentOffset];
-	CurrentOffset += Size;
+	void* Block = Header->GetBlockStart();
+	void* NextPtr = (void*)(((SIZE_T)Block) + ((SIZE_T)Header->NextOffset));
+	void* AlignedPtr = Align(NextPtr, Alignment);
+	SIZE_T NewOffset = ((SIZE_T)AlignedPtr - (SIZE_T)Block) + Size;
+	if (NewOffset + sizeof(FBlock) > (SIZE_T)Header->BlockSize)
+	{
+		return nullptr;
+	}
 
-	return Mem;
+	Header->NextOffset = (uint32)NewOffset;
+	return AlignedPtr;
 }
 
-void FLinearBlockAllocator::AllocateNewBlock(SIZE_T Size)
+bool FLinearBlockAllocator::RequiresCustomAllocation(SIZE_T Size, uint32 Alignment) const
 {
-	Size = Align(Size, LinearAllocatorBlockSize);
-	FBlock* PreviousHeader = CurrentHeader;
+	// Compute the worst case AlignmentPadding, which occurs if the Block pointer is aligned
+	// only to BlockAlignment and no higher.
+	SIZE_T StartingOffset = UE::LinearAllocator::BlockAlignment;
+	SIZE_T Offset = StartingOffset;
+	// Insert alignment padding
+	Offset = Align(Offset, Alignment);
+	// Add on the size.
+	Offset += Size;
+	// Add on the size of the Block header which is placed at the end of the allocated block.
+	Offset += sizeof(FBlock);
 
-	CurrentBlock = (uint8*)FMemory::Malloc(Size);
-	CurrentHeader = (FBlock*)(CurrentBlock + Size - sizeof(FBlock));
-	CurrentHeader->Next = nullptr;
-	CurrentHeader->LinearAllocatorBlockSize = Size;
+	SIZE_T RequiredBlockSize = Offset - StartingOffset;
+	return DefaultBlockSize < RequiredBlockSize;
+}
 
-	if (PreviousHeader)
+FLinearBlockAllocator::FBlock* FLinearBlockAllocator::AddDefaultBlock()
+{
+	FBlock* NewHeader = AllocateBlock(DefaultBlockSize, UE::LinearAllocator::BlockAlignment);
+	check(NewHeader);
+	if (!FirstHeader)
+	{
+		check(!LastHeader);
+		FirstHeader = NewHeader;
+		LastHeader = NewHeader;
+	}
+	else
 	{
-		PreviousHeader->Next = CurrentHeader;
+		check(LastHeader);
+		LastHeader->Next = NewHeader;
+		LastHeader = NewHeader;
 	}
+	return NewHeader;
+}
+
+FLinearBlockAllocator::FBlock* FLinearBlockAllocator::AddCustomBlock(SIZE_T UserSize, uint32 UserAlignment)
+{
+	uint32 CustomAlignment = FMath::RoundUpToPowerOfTwo(FMath::Max(UserAlignment, UE::LinearAllocator::BlockAlignment));
+	// User size was checked to fit within our smaller-type CustomBlockSize by the public Malloc functions
+	// before calling AddCustomBlock.
+	uint32 CustomBlockSize = (uint32)UserSize;
+	CustomBlockSize = FMath::Max((CustomBlockSize + (uint32)sizeof(FBlock)), MinBlockSize);
+	CustomBlockSize = Align(CustomBlockSize, UE::LinearAllocator::BlockAlignment);
 
-	CurrentOffset = 0;
-	CurrentBlockSize = Size - sizeof(FBlock);
-	TotalAllocated += Size;
+	FBlock* NewHeader = AllocateBlock(CustomBlockSize, CustomAlignment);
+	check(NewHeader);
+	// Insert custom blocks at the beginning of the linked list, to avoid wasting the remaining memory in LastBlock.
+	if (!FirstHeader)
+	{
+		check(!LastHeader);
+		FirstHeader = NewHeader;
+		LastHeader = NewHeader;
+	}
+	else
+	{
+		check(LastHeader);
+		NewHeader->Next = FirstHeader;
+		FirstHeader = NewHeader;
+	}
+	return NewHeader;
 }
 
-bool FLinearBlockAllocator::CanFit(SIZE_T Size, uint32 Alignment) const
+FLinearBlockAllocator::FBlock* FLinearBlockAllocator::AllocateBlock(uint32 BlockSize, uint32 BlockAlignment)
 {
-	return (CurrentBlockSize - Align(CurrentOffset, Alignment)) >= Size;
+	check(BlockSize >= MinBlockSize);
+	check(IsAligned(BlockSize, UE::LinearAllocator::BlockAlignment));
+	check(IsAligned(BlockAlignment, UE::LinearAllocator::BlockAlignment));
+
+	void* NewBlock = InnerMalloc->Malloc(BlockSize, BlockAlignment);
+	check(NewBlock != nullptr);
+	LLM(const bool bDeclareAllocs = (InnerMalloc == UE::Private::GMalloc) && FLowLevelMemTracker::IsEnabled());
+	LLM(if (bDeclareAllocs)
+	{
+		FLowLevelMemTracker::Get().OnLowLevelAlloc(ELLMTracker::Default, NewBlock, BlockSize,
+			ELLMTag::Untagged, ELLMAllocType::FMalloc);
+	});
+	TotalAllocated += BlockSize;
+
+	// Note that NewHeader is aligned to at least UE::LinearAllocator::BlockAlignment, because BlockAllocation,
+	// BlockSize, and sizeof(FBlock) are all aligned to at least UE::LinearAllocator::BlockAlignment.
+	FBlock* NewHeader = FBlock::GetHeaderFromBlockStart(NewBlock, BlockSize);
+	NewHeader->Next = nullptr;
+	NewHeader->BlockSize = BlockSize;
+	NewHeader->NextOffset = 0;
+	return NewHeader;
 }
 
-typedef FLinearBlockAllocator FPersistentLinearAllocator;
 
-#endif //~UE_ENABLE_LINEAR_VIRTUAL_ALLOCATOR
+FLinearBlockAllocatorThreadAccessor::FLinearBlockAllocatorThreadAccessor(FLinearBlockAllocator& InAllocator)
+	: Allocator(&InAllocator)
+{
+	Allocator->ReferenceCount.fetch_add(1, std::memory_order_relaxed);
+}
+
+FLinearBlockAllocatorThreadAccessor::~FLinearBlockAllocatorThreadAccessor()
+{
+	if (Allocator)
+	{
+		Allocator->ReferenceCount.fetch_sub(1, std::memory_order_relaxed);
+	}
+}
+
+FLinearBlockAllocatorThreadAccessor::FLinearBlockAllocatorThreadAccessor(
+	FLinearBlockAllocatorThreadAccessor&& Other)
+{
+	Allocator = Other.Allocator;
+	LastHeader = Other.LastHeader;
+	Other.Allocator = nullptr;
+	Other.LastHeader = nullptr;
+}
+
+void* FLinearBlockAllocatorThreadAccessor::Malloc(SIZE_T Size, uint32 Alignment)
+{
+	checkf(Size <= MAX_uint32 - sizeof(FLinearBlockAllocator::FBlock) * 2,
+		TEXT("FLinearBlockAllocator does not support allocations of size ~4GB or more. We store BlockSizes as uint32."));
+	// Malloc must not be called on cleared thread accessors (e.g. cleared by move constructor into another accessor)
+	check(Allocator);
+
+	Alignment = FMath::Max(Alignment, UE::LinearAllocator::DefaultAlignment);
+	void* Result = Allocator->TryAllocateFromBlock(Size, Alignment, LastHeader);
+	if (Result)
+	{
+		return Result;
+	}
+
+	FLinearBlockAllocator::FBlock* NewHeader;
+	{
+		FScopeLock FScopeLock(&Allocator->Lock);
+		if (!Allocator->RequiresCustomAllocation(Size, Alignment))
+		{
+			NewHeader = Allocator->AddDefaultBlock();
+			LastHeader = NewHeader;
+		}
+		else
+		{
+			NewHeader = Allocator->AddCustomBlock(Size, Alignment);
+			// Do not modify LastHeader, keep allocating from the old one
+		}
+	}
+	Result = Allocator->TryAllocateFromBlock(Size, Alignment, NewHeader);
+	check(Result);
+	return Result;
+}
+
+} // namespace UE
 
 FLinearAllocator& GetPersistentLinearAllocator()
 {
diff --git a/Engine/Source/Runtime/Core/Public/Memory/LinearAllocator.h b/Engine/Source/Runtime/Core/Public/Memory/LinearAllocator.h
--- a/Engine/Source/Runtime/Core/Public/Memory/LinearAllocator.h
+++ b/Engine/Source/Runtime/Core/Public/Memory/LinearAllocator.h
@@ -5,14 +5,127 @@
 #include "HAL/Platform.h"
 #include "HAL/UnrealMemory.h"
 #include "HAL/CriticalSection.h"
+#include "Templates/SharedPointerFwd.h" // ESPMode
+#include "Templates/UnrealTemplate.h"
+
+#include <atomic>
 
 #ifndef UE_ENABLE_LINEAR_VIRTUAL_ALLOCATOR
 #	define UE_ENABLE_LINEAR_VIRTUAL_ALLOCATOR PLATFORM_HAS_FPlatformVirtualMemoryBlock
 #endif
 
+namespace UE
+{
+
+class FLinearBlockAllocatorThreadAccessor;
+
+/**
+ * Provides Malloc operation that reduces per-malloc and per-free overhead by allocating in blocks, and not allowing
+ * individual calls to free. Does not free any malloc'd memory until the allocator goes out of scope.
+ * Intended use is dynamically allocated scratch data, or long-lived global allocations that never need to free data.
+ * Allocations at the end of each block waste memory: when the next requested allocation is too big for the current
+ * block, the remaining memory goes unused. To avoid large amounts of waste, allocations should be relatively small
+ * compared to blocksize.
+ *
+ * Memory for the blocks is allocated from InnerMalloc->Malloc. If InnerMalloc is nullptr, GMalloc is used.
+ * 
+ * The LinearBlockAllocator must be destroyed before its InnerMalloc is destroyed; the InnerMalloc pointer is
+ * not reference counted.
+ */
+class FLinearBlockAllocator : public FNoncopyable
+{
+public:
+	/** 
+	 * Constructor.
+	 * 
+	 * @param InInnerMalloc The allocator used to allocate blocks and internal memory. If null, uses GMalloc.
+	 * @param InBlockSize The size of each block. Rounded up to MinBlockSize and next power of 2. If 0, uses
+	 *        UE::LinearAllocator::DefaultBlockSize.
+	 * @param InThreadMode If ESPMode::ThreadSafe, the allocator locks a critical section around each allocation
+	 *        and can be used from any thread. If ESPMode::NotThreadSafe, the allocator must either be used only
+	 *        from a single thread, or cannot be accessed directly and must instead be accessed through
+	 *        FLinearBlockAllocatorThreadAccessor. Each FLinearBlockAllocatorThreadAccessor owns its own block.
+	 *        The unallocated memory in each threaded block is wasted and not available to other threads.
+	 */
+	CORE_API FLinearBlockAllocator(FMalloc* InInnerMalloc = nullptr, uint32 InBlockSize = 0,
+		ESPMode InThreadMode = ESPMode::ThreadSafe);
+	CORE_API ~FLinearBlockAllocator();
+	FLinearBlockAllocator(FLinearBlockAllocator&&) = delete;
+	FLinearBlockAllocator& operator=(FLinearBlockAllocator&&) = delete;
+
+	CORE_API void* Malloc(SIZE_T Size, uint32 Alignment = DEFAULT_ALIGNMENT);
+
+	/** Returns the size of allocated blocks, which is >= size of pointers returned from Malloc calls. */
+	CORE_API SIZE_T GetAllocatedMemorySize() const;
+
+protected:
+	struct FBlock
+	{
+		FBlock* Next; // Guarded by Lock
+		uint32 BlockSize; // Constant during threading
+		uint32 NextOffset; // Read/Written by owning thread or guarded by Lock in ESPMode::ThreadSafe
+
+		static FBlock* GetHeaderFromBlockStart(void* Block, uint32 BlockSize);
+		void* GetBlockStart();
+	};
+	constexpr static uint32 MinBlockSize = (uint32)(2 * sizeof(FBlock));
+
+	bool IsThreadSafe() const;
+	void* TryAllocateFromBlock(SIZE_T Size, uint32 Alignment, FBlock* BlockHeader);
+	bool RequiresCustomAllocation(SIZE_T Size, uint32 Alignment) const;
+	FBlock* AddDefaultBlock();
+	FBlock* AddCustomBlock(SIZE_T UserSize, uint32 UserAlignment);
+	FBlock* AllocateBlock(uint32 BlockSize, uint32 BlockAlignment);
+
+	mutable FCriticalSection Lock;
+	FMalloc* InnerMalloc = nullptr; // Pointer is constant during threading.
+	FBlock* FirstHeader = nullptr; // Guarded by Lock
+	FBlock* LastHeader = nullptr; // Guarded by Lock
+	SIZE_T TotalAllocated = 0; // Guarded by Lock
+	uint32 DefaultBlockSize = 0; // constant during threading
+	ESPMode ThreadMode = ESPMode::ThreadSafe; // constant during threading
+	std::atomic<int32> ReferenceCount{ 0 };
+
+	friend FLinearBlockAllocatorThreadAccessor;
+};
+
+/**
+ * Accessor for a LinearBlockAllocator when multithreading. Every thread using the Allocator must access it through
+ * its own FLinearBlockAllocatorThreadAccessor.
+ *
+ * Intended use is to share an allocator amongst the task threads used within a single function in e.g. ParallelFor.
+ *
+ * The lifetime of FLinearBlockAllocatorThreadAccessor must be within the lifetime of the Allocator, an assertion
+ * fires if FLinearBlockAllocatorThreadAccessor are still allocated when the Allocator is destructed. So the thread
+ * that destroys the FLinearBlockAllocator must wait for all other threads using the Allocator to destroy their
+ * FLinearBlockAllocatorThreadAccessor.
+ */
+class FLinearBlockAllocatorThreadAccessor : public FNoncopyable
+{
+public:
+	CORE_API explicit FLinearBlockAllocatorThreadAccessor(FLinearBlockAllocator& InAllocator);
+	CORE_API ~FLinearBlockAllocatorThreadAccessor();
+	CORE_API FLinearBlockAllocatorThreadAccessor(FLinearBlockAllocatorThreadAccessor&& Other);
+
+	// Copy constructor/assignment is disabled because each copy needs to get its own LastBlock from the allocator.
+	// Move assignment is disabled because it would waste the remaining available memory allocated for *this.
+	FLinearBlockAllocatorThreadAccessor& operator=(FLinearBlockAllocatorThreadAccessor&& Other) = delete;
+
+	CORE_API void* Malloc(SIZE_T Size, uint32 Alignment = DEFAULT_ALIGNMENT);
+
+private:
+	FLinearBlockAllocator* Allocator = nullptr;
+	FLinearBlockAllocator::FBlock* LastHeader = nullptr;
+};
+
+}
 
 #if UE_ENABLE_LINEAR_VIRTUAL_ALLOCATOR
 
+/**
+ * A more efficient LinearAllocator than UE::FLinearBlockAllocator, but one tuned for global use for Unreal's persistent
+ * core and engine allocations. This allocator should not be used by other systems.
+ */
 struct FLinearAllocator
 {
 	CORE_API FLinearAllocator(SIZE_T ReserveMemorySize);
@@ -21,7 +134,7 @@ struct FLinearAllocator
 		VirtualMemory.FreeVirtual();
 	}
 
-	CORE_API void* Allocate(SIZE_T Size, uint32 Alignment = 8);
+	CORE_API void* Allocate(SIZE_T Size, uint32 Alignment = DEFAULT_ALIGNMENT);
 
 	UE_FORCEINLINE_HINT SIZE_T GetAllocatedMemorySize() const
 	{
@@ -41,39 +154,15 @@ protected:
 
 #else
 
-struct FLinearBlockAllocator
+struct FLinearAllocator : public UE::FLinearBlockAllocator
 {
-	CORE_API FLinearBlockAllocator(SIZE_T);
-	CORE_API ~FLinearBlockAllocator();
-
-	CORE_API void* Allocate(SIZE_T Size, uint32 Alignment = 8);
-
-	UE_FORCEINLINE_HINT SIZE_T GetAllocatedMemorySize() const
+	CORE_API FLinearAllocator(SIZE_T ReserveMemorySize);
+	void* Allocate(SIZE_T Size, uint32 Alignment = DEFAULT_ALIGNMENT)
 	{
-		return TotalAllocated;
+		return Malloc(Size, Alignment);
 	}
-
-private:
-	struct FBlock
-	{
-		FBlock* Next;
-		SIZE_T BlockSize;
-	};
-
-	FCriticalSection Lock;
-	FBlock* FirstHeader = nullptr;
-	FBlock* CurrentHeader = nullptr;
-	uint8* CurrentBlock = nullptr;
-	SIZE_T CurrentOffset = 0;
-	SIZE_T CurrentBlockSize = 0;
-	SIZE_T TotalAllocated = 0;
-
-	void AllocateNewBlock(SIZE_T Size);
-	bool CanFit(SIZE_T Size, uint32 Alignment) const;
 };
 
-typedef FLinearBlockAllocator FLinearAllocator;
-
 #endif	//~UE_ENABLE_LINEAR_VIRTUAL_ALLOCATOR
 
 CORE_API FLinearAllocator& GetPersistentLinearAllocator();
diff --git a/Engine/Source/Runtime/Core/Tests/Memory/LinearAllocatorTest.cpp b/Engine/Source/Runtime/Core/Tests/Memory/LinearAllocatorTest.cpp
new file mode 100644
--- /dev/null
+++ b/Engine/Source/Runtime/Core/Tests/Memory/LinearAllocatorTest.cpp
@@ -0,0 +1,372 @@
+// Copyright Epic Games, Inc. All Rights Reserved.
+
+#if WITH_TESTS
+
+#include "Memory/LinearAllocator.h"
+
+#include "Async/ParallelFor.h"
+#include "Containers/Array.h"
+#include "Containers/Set.h"
+#include "Templates/AlignmentTemplates.h"
+#include "Tests/TestHarnessAdapter.h"
+
+class FInnerMallocForLinearAllocatorTest : public FMalloc
+{
+public:
+	virtual void* Malloc(SIZE_T Count, uint32 Alignment = DEFAULT_ALIGNMENT) override
+	{
+		void* Ptr = FMemory::Malloc(Count, Alignment);
+		Ptrs.Add(Ptr);
+		return Ptr;
+	}
+
+	virtual void* Realloc(void* Original, SIZE_T Count, uint32 Alignment = DEFAULT_ALIGNMENT) override
+	{
+		// LinearBlockAllocator does not use Realloc.
+		check(false);
+		return nullptr;
+	}
+
+	virtual void Free(void* Original)
+	{
+		if (Original)
+		{
+			int32 RemoveCount = Ptrs.Remove(Original);
+			check(RemoveCount == 1);
+		}
+		FMemory::Free(Original);
+	}
+	TSet<void*> Ptrs;
+};
+
+static bool DoPointersOverlap(void* A, SIZE_T ASize, void* B, SIZE_T BSize)
+{
+	// Overlapping: B is before end of A, and A is before end of B.
+	const UPTRINT PA = reinterpret_cast<UPTRINT>(A);
+	const UPTRINT PB = reinterpret_cast<UPTRINT>(B);
+	return (PB < PA + ASize) && (PA < PB + BSize);
+}
+
+TEST_CASE_NAMED(FLinearBlockAllocatorTest, "System::Core::Memory::LinearBlockAllocator::Base", "[ApplicationContextMask][EngineFilter]")
+{
+	SECTION("Single allocation is non-null and aligned to DefaultAlignment")
+	{
+		UE::FLinearBlockAllocator Allocator;
+		void* Ptr = Allocator.Malloc(16);
+		CHECK(Ptr != nullptr);
+		CHECK(IsAligned(Ptr, 8));
+	}
+
+	SECTION("Two sequential allocations do not overlap")
+	{
+		UE::FLinearBlockAllocator Allocator;
+		SIZE_T AllocSize = 32;
+		void* A = Allocator.Malloc(AllocSize);
+		void* B = Allocator.Malloc(AllocSize);
+		CHECK(A != nullptr);
+		CHECK(B != nullptr);
+		CHECK(!DoPointersOverlap(A, AllocSize, B, AllocSize));
+	}
+
+	SECTION("Allocations returned pointers are writable across the full requested size")
+	{
+		UE::FLinearBlockAllocator Allocator;
+		constexpr SIZE_T Size = 256;
+		uint8* Ptr = static_cast<uint8*>(Allocator.Malloc(Size));
+		REQUIRE(Ptr != nullptr);
+		// Write and read back every byte.
+		for (SIZE_T i = 0; i < Size; ++i)
+		{
+			Ptr[i] = static_cast<uint8>(i);
+		}
+		bool bAllCorrect = true;
+		for (SIZE_T i = 0; i < Size; ++i)
+		{
+			if (Ptr[i] != static_cast<uint8>(i))
+			{
+				bAllCorrect = false;
+				break;
+			}
+		}
+		CHECK(bAllCorrect);
+	}
+
+	SECTION("Custom alignment is honoured")
+	{
+		UE::FLinearBlockAllocator Allocator;
+		for (uint32 Alignment : {8u, 16u, 32u, 64u, 128u, 256u})
+		{
+			void* Ptr = Allocator.Malloc(1, Alignment);
+			REQUIRE(Ptr != nullptr);
+			CHECK(IsAligned(Ptr, Alignment));
+		}
+	}
+
+	SECTION("Alignment smaller than DefaultAlignment is promoted to DefaultAlignment")
+	{
+		UE::FLinearBlockAllocator Allocator;
+		// Requesting alignment of 1 should still give at least 8-byte alignment.
+		void* Ptr = Allocator.Malloc(1, 1);
+		REQUIRE(Ptr != nullptr);
+		CHECK(IsAligned(Ptr, 8));
+	}
+
+	SECTION("Large number of small allocations all succeed and are non-overlapping")
+	{
+		UE::FLinearBlockAllocator Allocator(nullptr, 256);
+		constexpr int32 Count = 100;
+		constexpr SIZE_T AllocSize = 8;
+		TArray<void*> Ptrs;
+		Ptrs.Reserve(Count);
+
+		for (int32 i = 0; i < Count; ++i)
+		{
+			void* Ptr = Allocator.Malloc(AllocSize);
+			REQUIRE(Ptr != nullptr);
+			Ptrs.Add(Ptr);
+		}
+
+		// Verify uniqueness — no two returned pointers overlap.
+		for (int32 i = 0; i < Count; ++i)
+		{
+			for (int32 j = i + 1; j < Count; ++j)
+			{
+				CHECK(!DoPointersOverlap(Ptrs[i], AllocSize, Ptrs[j], AllocSize));
+			}
+		}
+	}
+
+	SECTION("Allocation larger than default block size succeeds (custom block path)")
+	{
+		// Use a small explicit block size so a standard allocation forces the custom block path.
+		UE::FLinearBlockAllocator Allocator(nullptr, 256);
+		constexpr SIZE_T LargeSize = 512;
+		void* Ptr = Allocator.Malloc(LargeSize);
+		CHECK(Ptr != nullptr);
+		CHECK(IsAligned(Ptr, 8));
+		// Write to the full extent to confirm it is backed by real memory.
+		FMemory::Memset(Ptr, 0xAB, LargeSize);
+	}
+
+	SECTION("GetAllocatedMemorySize increases after allocations that require new blocks")
+	{
+		UE::FLinearBlockAllocator Allocator;
+		const SIZE_T Before = Allocator.GetAllocatedMemorySize();
+
+		// Exhaust one full default block.
+		constexpr uint32 SmallBlockSize = 256;
+		UE::FLinearBlockAllocator SmallBlockAllocator(nullptr, SmallBlockSize);
+		const SIZE_T BeforeSmall = SmallBlockAllocator.GetAllocatedMemorySize();
+
+		// Allocate enough to force a second block.
+		SmallBlockAllocator.Malloc(SmallBlockSize);
+		SmallBlockAllocator.Malloc(SmallBlockSize);
+
+		CHECK(SmallBlockAllocator.GetAllocatedMemorySize() > BeforeSmall);
+		// Baseline allocator with no allocations should remain at its initial size.
+		CHECK(Allocator.GetAllocatedMemorySize() == Before);
+	}
+
+	SECTION("Allocations from different blocks do not overlap")
+	{
+		// Force multiple blocks by using a tiny block size.
+		UE::FLinearBlockAllocator Allocator(nullptr, 64);
+		TArray<TPair<void*, SIZE_T>> Ranges;
+		constexpr SIZE_T AllocSize = 48;
+		constexpr int32 Count = 8;
+
+		for (int32 i = 0; i < Count; ++i)
+		{
+			void* Ptr = Allocator.Malloc(AllocSize);
+			REQUIRE(Ptr != nullptr);
+			Ranges.Add({Ptr, AllocSize});
+		}
+
+		for (int32 i = 0; i < Ranges.Num(); ++i)
+		{
+			for (int32 j = i + 1; j < Ranges.Num(); ++j)
+			{
+				CHECK(!DoPointersOverlap(Ranges[i].Key, Ranges[i].Value, Ranges[j].Key, Ranges[j].Value));
+			}
+		}
+	}
+
+	SECTION("Zero-size allocation returns non-null and does not corrupt adjacent allocation")
+	{
+		UE::FLinearBlockAllocator Allocator;
+		void* ZeroPtr = Allocator.Malloc(0);
+		CHECK(ZeroPtr != nullptr);
+
+		// The allocation after a zero-size one should still be valid and writable.
+		constexpr SIZE_T GuardSize = 64;
+		uint8* GuardPtr = static_cast<uint8*>(Allocator.Malloc(GuardSize));
+		REQUIRE(GuardPtr != nullptr);
+		FMemory::Memset(GuardPtr, 0xCD, GuardSize);
+		CHECK(GuardPtr[0] == 0xCD);
+		CHECK(GuardPtr[GuardSize - 1] == 0xCD);
+	}
+
+	SECTION("Allocator does not leak memory with small allocations")
+	{
+		FInnerMallocForLinearAllocatorTest InnerMalloc;
+		uint32 BlockSize = 128; // Allocate a small block size so we can trigger multiple blocks
+		uint32 AllocationSize = 30;
+		uint32 NumAllocs = 50;
+		CHECK(AllocationSize * NumAllocs > BlockSize);
+
+		{
+			UE::FLinearBlockAllocator Allocator(&InnerMalloc, BlockSize);
+			for (uint32 Index = 0; Index < NumAllocs; ++Index)
+			{
+				Allocator.Malloc(AllocationSize);
+			}
+		}
+		CHECK(InnerMalloc.Ptrs.IsEmpty());
+	}
+
+	SECTION("Allocator does not leak memory with custom allocations")
+	{
+		FInnerMallocForLinearAllocatorTest InnerMalloc;
+		uint32 BlockSize = 128; // Allocate a small block size to avoid using a lot of memory
+		uint32 AllocationSize = 256;
+		uint32 NumAllocs = 4;
+		CHECK(AllocationSize * NumAllocs > BlockSize);
+
+		{
+			UE::FLinearBlockAllocator Allocator(&InnerMalloc, BlockSize);
+			for (uint32 Index = 0; Index < NumAllocs; ++Index)
+			{
+				Allocator.Malloc(AllocationSize);
+			}
+		}
+		CHECK(InnerMalloc.Ptrs.IsEmpty());
+	}
+
+	SECTION("Allocator does not leak memory with interleaved small and custom allocations, small first")
+	{
+		FInnerMallocForLinearAllocatorTest InnerMalloc;
+		uint32 BlockSize = 128; // Allocate a small block size to avoid using a lot of memory
+		uint32 SmallAllocationSize = 30;
+		uint32 LargeAllocationSize = 256;
+		uint32 NumTrials = 4;
+
+		{
+			UE::FLinearBlockAllocator Allocator(&InnerMalloc, BlockSize);
+			for (uint32 Index = 0; Index < NumTrials; ++Index)
+			{
+				Allocator.Malloc(SmallAllocationSize);
+				Allocator.Malloc(LargeAllocationSize);
+			}
+		}
+		CHECK(InnerMalloc.Ptrs.IsEmpty());
+	}
+
+	SECTION("Allocator does not leak memory with interleaved small and custom allocations, large first")
+	{
+		FInnerMallocForLinearAllocatorTest InnerMalloc;
+		uint32 BlockSize = 128; // Allocate a small block size to avoid using a lot of memory
+		uint32 SmallAllocationSize = 30;
+		uint32 LargeAllocationSize = 256;
+		uint32 NumTrials = 4;
+
+		{
+			UE::FLinearBlockAllocator Allocator(&InnerMalloc, BlockSize);
+			for (uint32 Index = 0; Index < NumTrials; ++Index)
+			{
+				Allocator.Malloc(LargeAllocationSize);
+				Allocator.Malloc(SmallAllocationSize);
+			}
+		}
+		CHECK(InnerMalloc.Ptrs.IsEmpty());
+	}
+}
+
+TEST_CASE_NAMED(FLinearBlockAllocatorThreadAccessorTest,
+	"System::Core::Memory::LinearBlockAllocator::ThreadAccessor", "[ApplicationContextMask][EngineFilter]")
+{
+	SECTION("ThreadAccessor Malloc returns valid aligned pointer")
+	{
+		UE::FLinearBlockAllocator Allocator(nullptr, 0, ESPMode::NotThreadSafe);
+		UE::FLinearBlockAllocatorThreadAccessor Accessor(Allocator);
+		void* Ptr = Accessor.Malloc(32);
+		CHECK(Ptr != nullptr);
+		CHECK(IsAligned(Ptr, 8));
+	}
+
+	SECTION("ThreadAccessor write does not corrupt neighbouring accessor allocation")
+	{
+		UE::FLinearBlockAllocator Allocator(nullptr, 1024, ESPMode::NotThreadSafe);
+		UE::FLinearBlockAllocatorThreadAccessor AccessorA(Allocator);
+		UE::FLinearBlockAllocatorThreadAccessor AccessorB(Allocator);
+
+		constexpr SIZE_T Size = 128;
+		uint8* A = static_cast<uint8*>(AccessorA.Malloc(Size));
+		uint8* B = static_cast<uint8*>(AccessorB.Malloc(Size));
+		REQUIRE(A != nullptr);
+		REQUIRE(B != nullptr);
+
+		FMemory::Memset(A, 0xAA, Size);
+		FMemory::Memset(B, 0xBB, Size);
+
+		bool bACorrect = true, bBCorrect = true;
+		for (SIZE_T i = 0; i < Size; ++i)
+		{
+			if (A[i] != 0xAA) { bACorrect = false; }
+			if (B[i] != 0xBB) { bBCorrect = false; }
+		}
+		CHECK(bACorrect);
+		CHECK(bBCorrect);
+	}
+
+	SECTION("Move constructor transfers ownership and clears source")
+	{
+		UE::FLinearBlockAllocator Allocator(nullptr, 0, ESPMode::NotThreadSafe);
+		UE::FLinearBlockAllocatorThreadAccessor Original(Allocator);
+		// Allocate via Original to give it a cached block.
+		void* FirstPtr = Original.Malloc(16);
+		CHECK(FirstPtr != nullptr);
+
+		UE::FLinearBlockAllocatorThreadAccessor Moved(MoveTemp(Original));
+		// Moved accessor should work.
+		void* SecondPtr = Moved.Malloc(16);
+		CHECK(SecondPtr != nullptr);
+	}
+
+	SECTION("Multiple ThreadAccessors in parallel each produce non-overlapping allocations")
+	{
+		UE::FLinearBlockAllocator Allocator(nullptr, 0, ESPMode::NotThreadSafe);
+
+		constexpr int32 NumThreads = 4;
+		constexpr int32 AllocsPerThread = 256;
+		constexpr SIZE_T AllocSize = 16;
+
+		TArray<TArray<void*>> AllPtrs;
+		AllPtrs.SetNum(NumThreads);
+
+		ParallelFor(NumThreads, [&](int32 ThreadIdx)
+		{
+			UE::FLinearBlockAllocatorThreadAccessor Accessor(Allocator);
+			AllPtrs[ThreadIdx].Reserve(AllocsPerThread);
+			for (int32 i = 0; i < AllocsPerThread; ++i)
+			{
+				void* Ptr = Accessor.Malloc(AllocSize);
+				REQUIRE(Ptr != nullptr);
+				AllPtrs[ThreadIdx].Add(Ptr);
+			}
+		});
+
+		// Flatten and verify all pointers are unique.
+		TSet<void*> Seen;
+		for (const TArray<void*>& ThreadPtrs : AllPtrs)
+		{
+			for (void* Ptr : ThreadPtrs)
+			{
+				bool bAlreadyInSet = false;
+				Seen.Add(Ptr, &bAlreadyInSet);
+				CHECK_FALSE(bAlreadyInSet);
+			}
+		}
+	}
+}
+
+#endif // WITH_TESTS
