Add IPlatformFile::OpenMappedEx2. OpenMappedEx2 has more granular flags so its possible to better specify behavior when working with multiple handles to the same file.

OpenMappedEx has different behavior on unix-like and windows platform when using AllowWrite. This behavior is not changed to avoid breaking any existing platform specific code.  Specifically, unix-like platforms use OpenForWrite|OpenExisting and Windows platforms use OpenForWrite|AllowRead.

Additionally
- MaximumSize is now correctly set in all implemetations for both read and write.
- On unix-like platforms, mode flags are set consistent with other API calls on the platform.
- On unix-like platforms, O_CLOEXEC is set consistent with other API calls on the platform.

#Synced-CL 51330535

diff --git a/Engine/Source/Runtime/Core/Private/Android/AndroidPlatformFile.cpp b/Engine/Source/Runtime/Core/Private/Android/AndroidPlatformFile.cpp
--- a/Engine/Source/Runtime/Core/Private/Android/AndroidPlatformFile.cpp
+++ b/Engine/Source/Runtime/Core/Private/Android/AndroidPlatformFile.cpp
@@ -1712,16 +1712,48 @@ public:
 	}
 
 	virtual FOpenMappedResult OpenMappedEx(const TCHAR* Filename, EOpenReadFlags OpenOptions = EOpenReadFlags::None, int64 MaximumSize = 0) override
+	{
+		EOpenMappedFlags ActualOpenOptions = EOpenMappedFlags::None;
+		if (EnumHasAnyFlags(OpenOptions, EOpenReadFlags::AllowWrite))
+		{
+			ActualOpenOptions |= EOpenMappedFlags::OpenForWrite;
+			ActualOpenOptions |= EOpenMappedFlags::OpenExisting;
+		}
+
+		return OpenMappedEx2(Filename, ActualOpenOptions, MaximumSize);
+	}
+
+	virtual FOpenMappedResult OpenMappedEx2(const TCHAR* Filename, EOpenMappedFlags OpenOptions = EOpenMappedFlags::None, int64 MaximumSize = 0) override
 	{
 		FString LocalPath;
 		FString AssetPath;
 		const bool bForceAllowLocal = false;
 		PathToAndroidPaths(LocalPath, AssetPath, Filename, bForceAllowLocal);
 
-		const FString NormalizedFilename = LocalPath;
+		const FString NormalizedFilename = MoveTemp(LocalPath);
+
+		int Flags = O_RDONLY;
+		if (EnumHasAnyFlags(OpenOptions, EOpenMappedFlags::OpenForWrite))
+		{
+			Flags = O_RDWR;
+			if (!EnumHasAnyFlags(OpenOptions, EOpenMappedFlags::OpenExisting))
+			{
+				Flags |= O_CREAT;
+				Flags |= O_TRUNC;
+			}
+		}
+
+		int32 iHandle = -1;
+		if (EnumHasAnyFlags(OpenOptions, EOpenMappedFlags::OpenForWrite))
+		{
+			iHandle = open(TCHAR_TO_UTF8(*NormalizedFilename), Flags, S_IRUSR | S_IWUSR | S_IRGRP | S_IWGRP);
+		}
+		else
+		{
+			iHandle = open(TCHAR_TO_UTF8(*NormalizedFilename), Flags);
+		}
 
-		const int Flags = EnumHasAnyFlags(OpenOptions, EOpenReadFlags::AllowWrite) ? O_RDWR : O_RDONLY;
-		FAndroidUniqueFd Handle { open(TCHAR_TO_UTF8(*NormalizedFilename), Flags) };
+		FAndroidUniqueFd Handle{ iHandle };
 		if (!Handle.IsValid())
 		{
 			const int ErrNo = errno;
@@ -1747,7 +1779,56 @@ public:
 			return MakeError(MoveTemp(ErrorStr));
 		}
 
-		return MakeValue(new FAndroidMappedFileHandle(MoveTemp(Handle), FileInfo.st_size, NormalizedFilename));
+		if (EnumHasAnyFlags(OpenOptions, EOpenMappedFlags::OpenForWrite))
+		{
+			if (MaximumSize < 1)
+			{
+				if (EnumHasAnyFlags(OpenOptions, EOpenMappedFlags::OpenExisting))
+				{
+					if (FileInfo.st_size < 1)
+					{
+						return MakeError(TEXT("File size is 0"));
+					}
+
+					MaximumSize = FileInfo.st_size;
+				}
+				else
+				{
+					return MakeError(TEXT("No maximum size provided."));
+				}
+			}
+			else if (MaximumSize > FileInfo.st_size)
+			{
+				int TruncResult = 0;
+				do { TruncResult = __ftruncate(Handle.Get(), MaximumSize); } while (TruncResult < 0 && errno == EINTR);
+
+				if (TruncResult == -1)
+				{
+					const int ErrNo = errno;
+					FString ErrorStr = FString::Printf(TEXT("FAndroidPlatformFile::OpenMappedEx ftruncate failed: ('%s', Flags=0x%08X) failed: errno=%d (%s)"), *NormalizedFilename, Flags, ErrNo, UTF8_TO_TCHAR(strerror(ErrNo)));
+
+#if LOG_ANDROID_FILE
+					FPlatformMisc::LowLevelOutputDebugString(*ErrorStr);
+#endif
+
+					return MakeError(MoveTemp(ErrorStr));
+				}
+			}
+		}
+		else
+		{
+			if (FileInfo.st_size < 1)
+			{
+				return MakeError(TEXT("File size is 0"));
+			}
+
+			if (MaximumSize < 1 || MaximumSize > FileInfo.st_size)
+			{
+				MaximumSize = FileInfo.st_size;
+			}
+		}
+
+		return MakeValue(new FAndroidMappedFileHandle(MoveTemp(Handle), MaximumSize, NormalizedFilename));
 	}
 
 	virtual int64 FileSize(const TCHAR* Filename) override
diff --git a/Engine/Source/Runtime/Core/Private/Apple/ApplePlatformFile.cpp b/Engine/Source/Runtime/Core/Private/Apple/ApplePlatformFile.cpp
--- a/Engine/Source/Runtime/Core/Private/Apple/ApplePlatformFile.cpp
+++ b/Engine/Source/Runtime/Core/Private/Apple/ApplePlatformFile.cpp
@@ -696,7 +696,7 @@ IFileHandle* FApplePlatformFile::OpenWrite(const TCHAR* Filename, bool bAppend,
 	const FString NormalizedFilename = NormalizeFilename(Filename);
 
 	TRACE_PLATFORMFILE_BEGIN_OPEN(*NormalizedFilename);
-	int32 Handle = open(TCHAR_TO_UTF8(*NormalizedFilename), Flags, S_IRUSR | S_IWUSR | S_IRGRP | S_IWGRP | S_IROTH | S_IWOTH);
+	int32 Handle = open(TCHAR_TO_UTF8(*NormalizedFilename), Flags, GetModeFlagsForOpenWrite());
 	
 	if (Handle != -1)
 	{
@@ -751,48 +751,113 @@ IFileHandle* FApplePlatformFile::OpenWrite(const TCHAR* Filename, bool bAppend,
 }
 
 FOpenMappedResult FApplePlatformFile::OpenMappedEx(const TCHAR* Filename, EOpenReadFlags OpenOptions, int64 MaximumSize)
+{
+	EOpenMappedFlags ActualOpenOptions = EOpenMappedFlags::None;
+	if (EnumHasAnyFlags(OpenOptions, EOpenReadFlags::AllowWrite))
+	{
+		ActualOpenOptions |= EOpenMappedFlags::OpenForWrite;
+		ActualOpenOptions |= EOpenMappedFlags::OpenExisting;
+	}
+
+	return OpenMappedEx2(Filename, ActualOpenOptions, MaximumSize);
+}
+
+FOpenMappedResult FApplePlatformFile::OpenMappedEx2(const TCHAR* Filename, EOpenMappedFlags OpenOptions, int64 MaximumSize)
 {
 	const FString NormalizedFilename = NormalizeFilename(Filename);
-	
-	// check the read path
-	FString FinalPath = ConvertToPlatformPath(NormalizedFilename, false, false);
-	FILE* FP;
-	const char* OpenMode = EnumHasAnyFlags(OpenOptions, EOpenReadFlags::AllowWrite) ? "r+" : "r";
-	FP = fopen(TCHAR_TO_UTF8(*FinalPath), OpenMode);
-	if(FP == nullptr)
+
+	FString FinalPath;
+	const int32 Handle = OpenFileForMemMap(NormalizedFilename, OpenOptions, FinalPath);
+	if (Handle == -1)
 	{
-		// if not in the read path, check the private write path
-		FinalPath = ConvertToPlatformPath(NormalizedFilename, true, false);
-		FP = fopen(TCHAR_TO_UTF8(*FinalPath), OpenMode);
+		return MakeError(FString::Printf(TEXT("FApplePlatformFile::OpenMappedEx failed to open file '%s' in '%s' mode"), *NormalizedFilename, EnumHasAnyFlags(OpenOptions, EOpenMappedFlags::OpenForWrite) ? TEXT("write") : TEXT("read")));
+	}
 
-		if(FP == nullptr)
-		{
-			// if not in the private write path, check the public write path
-			FinalPath = ConvertToPlatformPath(NormalizedFilename, true, true);
-			FP = fopen(TCHAR_TO_UTF8(*FinalPath), OpenMode);
-		}
+	struct stat FileInfo;
+	FileInfo.st_size = -1;
+	const int StatResult = fstat(Handle, &FileInfo);
+	if (StatResult == -1)
+	{
+		close(Handle);
+		return MakeError(FString::Printf(TEXT("FApplePlatformFile::OpenMappedEx failed to get file info for file '%s'"), *NormalizedFilename));
 	}
 	
-	if (FP != nullptr)
+	if (EnumHasAnyFlags(OpenOptions, EOpenMappedFlags::OpenForWrite))
 	{
-		int32 Handle = fileno(FP);
-	
-		if (Handle != -1)
+		if (MaximumSize < 1)
 		{
-			struct stat FileInfo;
-			FileInfo.st_size = -1;
-			// check the read path
-			if(fstat(Handle, &FileInfo) == -1)
+			if (EnumHasAnyFlags(OpenOptions, EOpenMappedFlags::OpenExisting))
+			{
+				if (FileInfo.st_size < 1)
+				{
+					close(Handle);
+					return MakeError(TEXT("File size is 0"));
+				}
+
+				MaximumSize = FileInfo.st_size;
+			}
+			else
 			{
-				return MakeError(FString::Printf(TEXT("FApplePlatformFile::OpenMappedEx failed to get file info for file '%s'"), *NormalizedFilename));
+				close(Handle);
+				return MakeError(TEXT("No maximum size provided."));
 			}
-			uint64 FileSize = FileInfo.st_size;
+		}
+		else if (MaximumSize > FileInfo.st_size)
+		{
+			int TruncResult = 0;
+			do { TruncResult = ftruncate(Handle, MaximumSize); } while (TruncResult < 0 && errno == EINTR);
 
-			return MakeValue(new FMappedFileHandle(Handle, FileSize, FinalPath));
+			if (TruncResult == -1)
+			{
+				close(Handle);
+				return MakeError(FString::Printf(TEXT("FApplePlatformFile::OpenMappedEx ftruncate failed to extend file '%s'"), *NormalizedFilename));
+			}
 		}
 	}
-	
-	return MakeError(FString::Printf(TEXT("FApplePlatformFile::OpenMappedEx failed to open file '%s' in '%s' mode"), *NormalizedFilename, EnumHasAnyFlags(OpenOptions, EOpenReadFlags::AllowWrite) ? TEXT("write") : TEXT("read")));
+	else
+	{
+		if (FileInfo.st_size < 1)
+		{
+			close(Handle);
+			return MakeError(TEXT("File size is 0"));
+		}
+
+		if (MaximumSize < 1 || MaximumSize > FileInfo.st_size)
+		{
+			MaximumSize = FileInfo.st_size;
+		}
+	}
+
+	return MakeValue(new FMappedFileHandle(Handle, MaximumSize, FinalPath));
+}
+
+int32 FApplePlatformFile::OpenFileForMemMap(const FString& NormalizedFilename, EOpenMappedFlags OpenOptions, FString& OutFinalPath)
+{
+	OutFinalPath = NormalizedFilename;
+
+	int Flags = O_RDONLY | O_CLOEXEC;
+	if (EnumHasAnyFlags(OpenOptions, EOpenMappedFlags::OpenForWrite))
+	{
+		Flags = O_RDWR | O_CLOEXEC;
+		if (!EnumHasAnyFlags(OpenOptions, EOpenMappedFlags::OpenExisting))
+		{
+			Flags |= O_CREAT;
+			Flags |= O_TRUNC;
+		}
+	}
+
+	int32 Handle = -1;
+
+	if (EnumHasAnyFlags(OpenOptions, EOpenMappedFlags::OpenForWrite))
+	{
+		Handle = open(TCHAR_TO_UTF8(*NormalizedFilename), Flags, GetModeFlagsForOpenWrite());
+	}
+	else
+	{
+		Handle = open(TCHAR_TO_UTF8(*NormalizedFilename), Flags);
+	}
+
+	return Handle;
 }
 
 bool FApplePlatformFile::DirectoryExists(const TCHAR* Directory)
@@ -927,6 +992,11 @@ int32 FApplePlatformFile::Stat(const TCHAR* Filename, struct stat* OutFileInfo)
 	return stat(TCHAR_TO_UTF8(*NormalizeFilename(Filename)), OutFileInfo);
 }
 
+int FApplePlatformFile::GetModeFlagsForOpenWrite() const
+{
+	return S_IRUSR | S_IWUSR | S_IRGRP | S_IWGRP | S_IROTH | S_IWOTH;
+}
+
 int32 AppleOpenCommon(const TCHAR* Filename)
 {
 	TRACE_PLATFORMFILE_BEGIN_OPEN(Filename);
diff --git a/Engine/Source/Runtime/Core/Private/GenericPlatform/GenericPlatformFile.cpp b/Engine/Source/Runtime/Core/Private/GenericPlatform/GenericPlatformFile.cpp
--- a/Engine/Source/Runtime/Core/Private/GenericPlatform/GenericPlatformFile.cpp
+++ b/Engine/Source/Runtime/Core/Private/GenericPlatform/GenericPlatformFile.cpp
@@ -539,9 +539,17 @@ IAsyncReadFileHandle* IPlatformFile::OpenAsyncRead(const TCHAR* Filename, bool b
 
 FOpenMappedResult IPlatformFile::OpenMappedEx(const TCHAR* Filename, EOpenReadFlags OpenOptions, int64 MaximumSize)
 {
+	// Note: cannot forward to OpenMappedEx2 here. OpenMappedEx has inconsistent behavior on different platforms,
+	// so the call must be forwarded in the platform implementation
+
 	return MakeError(TEXT("OpenMappedEx is not implemented on this platform"));
 }
 
+FOpenMappedResult IPlatformFile::OpenMappedEx2(const TCHAR* Filename, EOpenMappedFlags OpenOptions, int64 MaximumSize)
+{
+	return MakeError(TEXT("OpenMappedEx2 is not implemented on this platform"));
+}
+
 FFileOpenResult IPlatformFile::OpenWrite(const TCHAR* Filename, EOpenWriteFlags Flags)
 {
 	const bool bAppend = EnumHasAnyFlags(Flags, EOpenWriteFlags::Append);
diff --git a/Engine/Source/Runtime/Core/Private/GenericPlatform/GenericPlatformFileTest.cpp b/Engine/Source/Runtime/Core/Private/GenericPlatform/GenericPlatformFileTest.cpp
--- a/Engine/Source/Runtime/Core/Private/GenericPlatform/GenericPlatformFileTest.cpp
+++ b/Engine/Source/Runtime/Core/Private/GenericPlatform/GenericPlatformFileTest.cpp
@@ -31,7 +31,7 @@ bool FMMapFileReadTest::RunTest(const FString& InParameter)
 	const bool bMakeTree = true;
 	UTEST_TRUE("Making directory tree", IFileManager::Get().MakeDirectory(*TestDirectory, bMakeTree));
 
-	// Create a dummy file to read later with OpenMappedEx
+	// Create a dummy file to read later with OpenMappedEx2
 	const char FileContent[] = "Temp file to read.";
 	IFileHandle* FileHandle = nullptr;
 	IPlatformFile& PlatformFile = FPlatformFileManager::Get().GetPlatformFile();
@@ -47,7 +47,7 @@ bool FMMapFileReadTest::RunTest(const FString& InParameter)
 	FileHandle->Flush();
 	delete FileHandle;
 
-	FOpenMappedResult ResMap = FPlatformFileManager::Get().GetPlatformFile().OpenMappedEx(*TempFileToRead);
+	FOpenMappedResult ResMap = FPlatformFileManager::Get().GetPlatformFile().OpenMappedEx2(*TempFileToRead);
 	TUniquePtr<IMappedFileHandle> Handle = ResMap.HasError() ? nullptr : ResMap.StealValue();
 
 	if (Handle)
@@ -55,7 +55,14 @@ bool FMMapFileReadTest::RunTest(const FString& InParameter)
 		IMappedFileRegion* Region = Handle->MapRegion();
 
 		int64 Size = Region->GetMappedSize();
+
+		UTEST_TRUE("Mapped size valid", Size >= sizeof(FileContent));
+
 		const char* Data = (const char*)Region->GetMappedPtr();
+
+		const bool bMemcmpSuccess = FMemory::Memcmp(FileContent, Data, sizeof(FileContent)) == 0;
+		UTEST_TRUE("Mapped data valid", bMemcmpSuccess);
+
 		delete Region;
 
 		return true;
@@ -64,7 +71,7 @@ bool FMMapFileReadTest::RunTest(const FString& InParameter)
 	// Some platforms might not implement OpenMappedEx so the test should not fail
 	if (ResMap.HasError())
 	{
-		if (ResMap.GetError().GetMessage() == TEXT("OpenMappedEx is not implemented on this platform"))
+		if (ResMap.GetError().GetMessage() == TEXT("OpenMappedEx2 is not implemented on this platform"))
 		{
 			return true;
 		}
@@ -89,37 +96,58 @@ bool FMMapFileWriteTest::RunTest(const FString& InParameter)
 	const bool bMakeTree = true;
 	UTEST_TRUE("Making directory tree", IFileManager::Get().MakeDirectory(*TestDirectory, bMakeTree))
 
-	FOpenMappedResult ResMap = FPlatformFileManager::Get().GetPlatformFile().OpenMappedEx(*TempFileToWrite, IPlatformFile::EOpenReadFlags::AllowWrite, MapAddressSpaceSize);
-	TUniquePtr<IMappedFileHandle> Handle = ResMap.HasError() ? nullptr : ResMap.StealValue();
+	const char FileContent[] = "Temp file to write.";
 
-	if (Handle)
+	FOpenMappedResult ResMap = FPlatformFileManager::Get().GetPlatformFile().OpenMappedEx2(
+		*TempFileToWrite, IPlatformFile::EOpenMappedFlags::OpenForWrite, MapAddressSpaceSize);
+
+	// Some platforms might not implement OpenMappedEx so the test should not fail
+	if (ResMap.HasError())
 	{
-		IMappedFileRegion* Region = Handle->MapRegion(0, MapAddressSpaceSize, EMappedFileFlags::EFileWritable);
+		const bool bNotImplemented = ResMap.GetError().GetMessage() == TEXT("OpenMappedEx2 is not implemented on this platform");
+		UTEST_TRUE("OpenMappedEx2 failed but is implemented for the platform", bNotImplemented);
 
-		const char FileContent[] = "Temp file to write.";
+		return bNotImplemented;
+	}
+
+	{
+		TUniquePtr<IMappedFileHandle> Handle = ResMap.StealValue();
+
+		IMappedFileRegion* Region = Handle->MapRegion(0, MapAddressSpaceSize, EMappedFileFlags::EFileWritable);
 
 		uint8* MapPtr = (uint8*)Region->GetMappedPtr();
 		uint64 MapSize = Region->GetMappedSize();
-		
+
 		check(sizeof(FileContent) <= MapSize);
 		FMemory::Memcpy(MapPtr, (uint8*)FileContent, sizeof(FileContent));
 		Handle->Flush();
 
 		delete Region;
-
-		return true;
 	}
 
-	// Some platforms might not implement OpenMappedEx so the test should not fail
-	if (ResMap.HasError())
+	ResMap = FPlatformFileManager::Get().GetPlatformFile().OpenMappedEx2(
+		*TempFileToWrite, IPlatformFile::EOpenMappedFlags::OpenForWrite | IPlatformFile::EOpenMappedFlags::OpenExisting, MapAddressSpaceSize);
+
+	UTEST_TRUE("OpenMappedEx2 OpenForWrite | OpenExisting succeed", ResMap.HasValue());
+
 	{
-		if (ResMap.GetError().GetMessage() == TEXT("OpenMappedEx is not implemented on this platform"))
-		{
-			return true;
-		}
+		TUniquePtr<IMappedFileHandle> Handle = ResMap.StealValue();
+
+		IMappedFileRegion* Region = Handle->MapRegion();
+
+		int64 Size = Region->GetMappedSize();
+
+		UTEST_TRUE("Mapped size valid", Size >= sizeof(FileContent));
+
+		const char* Data = (const char*)Region->GetMappedPtr();
+
+		const bool bMemcmpSuccess = FMemory::Memcmp(FileContent, Data, sizeof(FileContent)) == 0;
+		UTEST_TRUE("Mapped data valid", bMemcmpSuccess);
+
+		delete Region;
 	}
 
-	return false;
+	return true;
 }
 
 
diff --git a/Engine/Source/Runtime/Core/Private/HAL/IPlatformFileLogWrapper.cpp b/Engine/Source/Runtime/Core/Private/HAL/IPlatformFileLogWrapper.cpp
--- a/Engine/Source/Runtime/Core/Private/HAL/IPlatformFileLogWrapper.cpp
+++ b/Engine/Source/Runtime/Core/Private/HAL/IPlatformFileLogWrapper.cpp
@@ -92,3 +92,8 @@ FOpenMappedResult FLoggedPlatformFile::OpenMappedEx(const TCHAR* Filename, EOpen
 {
 	return LowerLevel->OpenMappedEx(Filename, OpenOptions, MaximumSize);
 }
+
+FOpenMappedResult FLoggedPlatformFile::OpenMappedEx2(const TCHAR* Filename, EOpenMappedFlags OpenOptions, int64 MaximumSize)
+{
+	return LowerLevel->OpenMappedEx2(Filename, OpenOptions, MaximumSize);
+}
diff --git a/Engine/Source/Runtime/Core/Private/HAL/IPlatformFileOpenLogWrapper.cpp b/Engine/Source/Runtime/Core/Private/HAL/IPlatformFileOpenLogWrapper.cpp
--- a/Engine/Source/Runtime/Core/Private/HAL/IPlatformFileOpenLogWrapper.cpp
+++ b/Engine/Source/Runtime/Core/Private/HAL/IPlatformFileOpenLogWrapper.cpp
@@ -111,6 +111,11 @@ FOpenMappedResult FPlatformFileOpenLog::OpenMappedEx(const TCHAR* Filename, EOpe
 	return LowerLevel->OpenMappedEx(Filename, OpenOptions, MaximumSize);
 }
 
+FOpenMappedResult FPlatformFileOpenLog::OpenMappedEx2(const TCHAR* Filename, EOpenMappedFlags OpenOptions, int64 MaximumSize)
+{
+	return LowerLevel->OpenMappedEx2(Filename, OpenOptions, MaximumSize);
+}
+
 void FPlatformFileOpenLog::AddLabelInternal(const TCHAR* Fmt, ...)
 {
 	if (GPlatformFileOpenLog)
diff --git a/Engine/Source/Runtime/Core/Private/IOS/IOSPlatformFile.cpp b/Engine/Source/Runtime/Core/Private/IOS/IOSPlatformFile.cpp
--- a/Engine/Source/Runtime/Core/Private/IOS/IOSPlatformFile.cpp
+++ b/Engine/Source/Runtime/Core/Private/IOS/IOSPlatformFile.cpp
@@ -822,7 +822,7 @@ IFileHandle* FIOSPlatformFile::OpenWrite(const TCHAR* Filename, bool bAppend, bo
 		Flags |= O_WRONLY;
 	}
 	FString IOSFilename = ConvertToPlatformPath(NormalizeFilename(Filename), true, bCreatePublicFiles);
-	int32 Handle = open(TCHAR_TO_UTF8(*IOSFilename), Flags, S_IRUSR | S_IWUSR | S_IRGRP | S_IWGRP | S_IROTH | S_IWOTH);
+	int32 Handle = open(TCHAR_TO_UTF8(*IOSFilename), Flags, GetModeFlagsForOpenWrite());
 
 	if (Handle != -1)
 	{
@@ -1073,6 +1073,51 @@ FString FIOSPlatformFile::ConvertToPlatformPath(const FString& Filename, bool bF
 	return Result;
 }
 
+int32 FIOSPlatformFile::OpenFileForMemMap(const FString& NormalizedFilename, EOpenMappedFlags OpenOptions, FString& OutFinalPath)
+{
+	int Flags = O_RDONLY;
+	if (EnumHasAnyFlags(OpenOptions, EOpenMappedFlags::OpenForWrite))
+	{
+		Flags = O_RDWR;
+		if (!EnumHasAnyFlags(OpenOptions, EOpenMappedFlags::OpenExisting))
+		{
+			Flags |= O_CREAT;
+			Flags |= O_TRUNC;
+		}
+	}
+
+	int32 Handle = -1;
+
+	if (EnumHasAnyFlags(OpenOptions, EOpenMappedFlags::OpenForWrite))
+	{
+		FString FinalPath = ConvertToPlatformPath(NormalizedFilename, true, bCreatePublicFiles);
+		Handle = open(TCHAR_TO_UTF8(*FinalPath), Flags, GetModeFlagsForOpenWrite());
+		OutFinalPath = MoveTemp(FinalPath);
+	}
+	else
+	{
+		// check the read path
+		FString FinalPath = ConvertToPlatformPath(NormalizedFilename, false, false);
+		Handle = open(TCHAR_TO_UTF8(*FinalPath), Flags);
+		if (Handle == -1)
+		{
+			// if not in the read path, check the private write path
+			FinalPath = ConvertToPlatformPath(NormalizedFilename, true, false);
+			Handle = open(TCHAR_TO_UTF8(*FinalPath), Flags);
+
+			if (Handle == -1)
+			{
+				// if not in the private write path, check the public write path
+				FinalPath = ConvertToPlatformPath(NormalizedFilename, true, true);
+				Handle = open(TCHAR_TO_UTF8(*FinalPath), Flags);
+			}
+		}
+		OutFinalPath = MoveTemp(FinalPath);
+	}
+
+	return Handle;
+}
+
 FString FIOSPlatformFile::GetSharedContainerPath(const TCHAR* AppGroupID)
 {
 	SCOPED_AUTORELEASE_POOL;
diff --git a/Engine/Source/Runtime/Core/Private/Unix/UnixPlatformFile.cpp b/Engine/Source/Runtime/Core/Private/Unix/UnixPlatformFile.cpp
--- a/Engine/Source/Runtime/Core/Private/Unix/UnixPlatformFile.cpp
+++ b/Engine/Source/Runtime/Core/Private/Unix/UnixPlatformFile.cpp
@@ -1306,11 +1306,42 @@ IFileHandle* FUnixPlatformFile::OpenWrite(const TCHAR* Filename, bool bAppend, b
 }
 
 FOpenMappedResult FUnixPlatformFile::OpenMappedEx(const TCHAR* Filename, EOpenReadFlags OpenOptions, int64 MaximumSize)
+{
+	EOpenMappedFlags ActualOpenOptions = EOpenMappedFlags::None;
+	if (EnumHasAnyFlags(OpenOptions, EOpenReadFlags::AllowWrite))
+	{
+		ActualOpenOptions |= EOpenMappedFlags::OpenForWrite;
+		ActualOpenOptions |= EOpenMappedFlags::OpenExisting;
+	}
+
+	return OpenMappedEx2(Filename, ActualOpenOptions, MaximumSize);
+}
+
+FOpenMappedResult FUnixPlatformFile::OpenMappedEx2(const TCHAR* Filename, EOpenMappedFlags OpenOptions, int64 MaximumSize)
 {
 	const FString NormalizedFilename = NormalizeFilename(Filename, false);
 
-	const int Flags = EnumHasAnyFlags(OpenOptions, EOpenReadFlags::AllowWrite) ? O_RDWR : O_RDONLY;
-	const int32 Handle = open(TCHAR_TO_UTF8(*NormalizedFilename), Flags);
+	int Flags = O_RDONLY | O_CLOEXEC;
+	if (EnumHasAnyFlags(OpenOptions, EOpenMappedFlags::OpenForWrite))
+	{
+		Flags = O_RDWR | O_CLOEXEC;
+		if (!EnumHasAnyFlags(OpenOptions, EOpenMappedFlags::OpenExisting))
+		{
+			Flags |= O_CREAT;
+			Flags |= O_TRUNC;
+		}
+	}
+
+	int32 Handle = -1;
+	if (EnumHasAnyFlags(OpenOptions, EOpenMappedFlags::OpenForWrite))
+	{
+		Handle = open(TCHAR_TO_UTF8(*NormalizedFilename), Flags, S_IRUSR | S_IWUSR | S_IRGRP | S_IWGRP | S_IROTH);
+	}
+	else
+	{
+		Handle = open(TCHAR_TO_UTF8(*NormalizedFilename), Flags);
+	}
+
 	if (Handle == -1)
 	{
 		const int ErrNo = errno;
@@ -1319,7 +1350,7 @@ FOpenMappedResult FUnixPlatformFile::OpenMappedEx(const TCHAR* Filename, EOpenRe
 
 		return MakeError(MoveTemp(ErrorStr));
 	}
-	
+
 	struct stat FileInfo;
 	FileInfo.st_size = -1;
 	const int StatResult = fstat(Handle, &FileInfo);
@@ -1329,10 +1360,61 @@ FOpenMappedResult FUnixPlatformFile::OpenMappedEx(const TCHAR* Filename, EOpenRe
 		FString ErrorStr = FString::Printf(TEXT("stat('%s', Flags=0x%08X) failed: errno=%d (%s)"), *NormalizedFilename, Flags, ErrNo, UTF8_TO_TCHAR(strerror(ErrNo)));
 		UE_LOG_UNIX_FILE(Warning, TEXT("%s"), *ErrorStr);
 
+		close(Handle);
 		return MakeError(MoveTemp(ErrorStr));
 	}
-	
-	return MakeValue(new FUnixMappedFileHandle(Handle, FileInfo.st_size, NormalizedFilename));
+
+	if (EnumHasAnyFlags(OpenOptions, EOpenMappedFlags::OpenForWrite))
+	{
+		if (MaximumSize < 1)
+		{
+			if (EnumHasAnyFlags(OpenOptions, EOpenMappedFlags::OpenExisting))
+			{
+				if (FileInfo.st_size < 1)
+				{
+					close(Handle);
+					return MakeError(TEXT("File size is 0"));
+				}
+
+				MaximumSize = FileInfo.st_size;
+			}
+			else
+			{
+				close(Handle);
+				return MakeError(TEXT("No maximum size provided."));
+			}
+		}
+		else if (MaximumSize > FileInfo.st_size)
+		{
+			int TruncResult = 0;
+			do { TruncResult = ftruncate(Handle, MaximumSize); } while (TruncResult < 0 && errno == EINTR);
+
+			if (TruncResult == -1)
+			{
+				int ErrNo = errno;
+				FString ErrorStr = FString::Printf(TEXT("ftruncate() failed for '%s': errno=%d (%s)"), *NormalizedFilename, ErrNo, UTF8_TO_TCHAR(strerror(ErrNo)));
+				UE_LOG_UNIX_FILE(Warning, TEXT("%s"), *ErrorStr);
+
+				close(Handle);
+				return MakeError(MoveTemp(ErrorStr));
+			}
+		}
+	}
+	else
+	{
+		if (FileInfo.st_size < 1)
+		{
+			close(Handle);
+			return MakeError(TEXT("File size is 0"));
+		}
+
+		if (MaximumSize < 1 || MaximumSize > FileInfo.st_size)
+		{
+			MaximumSize = FileInfo.st_size;
+		}
+	}
+
+	return MakeValue(new FUnixMappedFileHandle(Handle, MaximumSize, NormalizedFilename));
 }
 
 bool FUnixPlatformFile::DirectoryExists(const TCHAR* Directory)
diff --git a/Engine/Source/Runtime/Core/Private/Windows/WindowsPlatformFile.cpp b/Engine/Source/Runtime/Core/Private/Windows/WindowsPlatformFile.cpp
--- a/Engine/Source/Runtime/Core/Private/Windows/WindowsPlatformFile.cpp
+++ b/Engine/Source/Runtime/Core/Private/Windows/WindowsPlatformFile.cpp
@@ -1675,29 +1675,69 @@ IFileHandle* FWindowsPlatformFile::OpenWrite(const TCHAR* Filename, bool bAppend
 
 FOpenMappedResult FWindowsPlatformFile::OpenMappedEx(const TCHAR* Filename, EOpenReadFlags OpenOptions, int64 MaximumSize)
 {
-	int64 Size = FileSize(Filename);
+	EOpenMappedFlags ActualOpenOptions = EOpenMappedFlags::None;
+	if (EnumHasAnyFlags(OpenOptions, EOpenReadFlags::AllowWrite))
+	{
+		ActualOpenOptions |= EOpenMappedFlags::OpenForWrite;
+		ActualOpenOptions |= EOpenMappedFlags::AllowRead;
+	}
+	if (EnumHasAnyFlags(OpenOptions, EOpenReadFlags::AllowDelete))
+	{
+		ActualOpenOptions |= EOpenMappedFlags::AllowDelete;
+	}
+
+	return OpenMappedEx2(Filename, ActualOpenOptions, MaximumSize);
+}
+
+FOpenMappedResult FWindowsPlatformFile::OpenMappedEx2(const TCHAR* Filename, EOpenMappedFlags OpenOptions, int64 MaximumSize)
+{
 	uint32  Access = GENERIC_READ;
 	uint32  WinFlags = FILE_SHARE_READ |
-		(EnumHasAnyFlags(OpenOptions, EOpenReadFlags::AllowDelete) ? FILE_SHARE_DELETE : 0) |
+		(EnumHasAnyFlags(OpenOptions, EOpenMappedFlags::AllowWrite) ? FILE_SHARE_WRITE : 0) |
+		(EnumHasAnyFlags(OpenOptions, EOpenMappedFlags::AllowDelete) ? FILE_SHARE_DELETE : 0) |
 		GetFileShareDeleteFlagIfEnabled();
 	uint32  Create = OPEN_EXISTING;
 	DWORD	Protect = PAGE_READONLY;
 
-	if (EnumHasAnyFlags(OpenOptions, EOpenReadFlags::AllowWrite))
+	if (EnumHasAnyFlags(OpenOptions, EOpenMappedFlags::OpenForWrite))
 	{
 		Access = GENERIC_READ | GENERIC_WRITE;
-		WinFlags = FILE_SHARE_READ | FILE_SHARE_WRITE |
-			(EnumHasAnyFlags(OpenOptions, EOpenReadFlags::AllowDelete) ? FILE_SHARE_DELETE : 0) |
+		WinFlags = 
+			(EnumHasAnyFlags(OpenOptions, EOpenMappedFlags::AllowRead) ? FILE_SHARE_READ : 0) |
+			(EnumHasAnyFlags(OpenOptions, EOpenMappedFlags::AllowDelete) ? FILE_SHARE_DELETE : 0) |
 			GetFileShareDeleteFlagIfEnabled();
-		Create = CREATE_ALWAYS;
+		Create = EnumHasAnyFlags(OpenOptions, EOpenMappedFlags::OpenExisting) ? OPEN_EXISTING : CREATE_ALWAYS;
 		Protect = PAGE_READWRITE;
-		Size = MaximumSize;
+
+		if (MaximumSize < 1)
+		{
+			if (EnumHasAnyFlags(OpenOptions, EOpenMappedFlags::OpenExisting))
+			{
+				const int64 Size = FileSize(Filename);
+				if (Size < 1)
+				{
+					return MakeError(TEXT("File size is 0 or the file does not exist and no maximum size provided."));
+				}
+
+				MaximumSize = Size;
+			}
+			else
+			{
+				return MakeError(TEXT("No maximum size provided."));
+			}
+		}
 	}
 	else
 	{
+		const int64 Size = FileSize(Filename);
 		if (Size < 1)
 		{
-			return MakeError(TEXT("File size is 0"));
+			return MakeError(TEXT("File size is 0 or the file does not exist"));
+		}
+
+		if (MaximumSize < 1 || MaximumSize > Size)
+		{
+			MaximumSize = Size;
 		}
 	}
 
@@ -1714,10 +1754,7 @@ FOpenMappedResult FWindowsPlatformFile::OpenMappedEx(const TCHAR* Filename, EOpe
 		TRACE_PLATFORMFILE_FAIL_OPEN(Filename);
 		return MakeError(FString::Printf(TEXT("Failed to open file '%s'"), *NormalizedFilename));
 	}
-	if (MaximumSize == 0)
-	{
-		MaximumSize = Size;
-	}
+
 	HANDLE MappingHandle = CreateFileMapping(Handle, NULL, Protect, (DWORD)(MaximumSize >> 32), (DWORD)(MaximumSize), NULL);
 	if (MappingHandle == NULL)
 	{
@@ -1741,7 +1778,7 @@ FOpenMappedResult FWindowsPlatformFile::OpenMappedEx(const TCHAR* Filename, EOpe
 #endif
 		return MakeError(FString::Printf(TEXT("Failed to create file mapping for a file'%s'"), *NormalizedFilename));
 	}
-	return MakeValue(new FMappedFileHandleWindows(Handle, MappingHandle, Size, Filename));
+	return MakeValue(new FMappedFileHandleWindows(Handle, MappingHandle, MaximumSize, Filename));
 }
 
 bool FWindowsPlatformFile::DirectoryExists(const TCHAR* Directory)
diff --git a/Engine/Source/Runtime/Core/Private/Windows/WindowsPlatformFilePrivate.h b/Engine/Source/Runtime/Core/Private/Windows/WindowsPlatformFilePrivate.h
--- a/Engine/Source/Runtime/Core/Private/Windows/WindowsPlatformFilePrivate.h
+++ b/Engine/Source/Runtime/Core/Private/Windows/WindowsPlatformFilePrivate.h
@@ -70,6 +70,7 @@ public:
 	virtual IFileHandle* OpenReadNoBuffering(const TCHAR* Filename, bool bAllowWrite = false) override;
 	virtual IFileHandle* OpenWrite(const TCHAR* Filename, bool bAppend = false, bool bAllowRead = false) override;
 	virtual FOpenMappedResult OpenMappedEx(const TCHAR* Filename, EOpenReadFlags OpenOptions = EOpenReadFlags::None, int64 MaximumSize = 0) override;
+	virtual FOpenMappedResult OpenMappedEx2(const TCHAR* Filename, EOpenMappedFlags OpenOptions = EOpenMappedFlags::None, int64 MaximumSize = 0) override;
 	virtual bool DirectoryExists(const TCHAR* Directory) override;
 	virtual bool CreateDirectory(const TCHAR* Directory) override;
 	virtual bool DeleteDirectory(const TCHAR* Directory) override;
diff --git a/Engine/Source/Runtime/Core/Public/Apple/ApplePlatformFile.h b/Engine/Source/Runtime/Core/Public/Apple/ApplePlatformFile.h
--- a/Engine/Source/Runtime/Core/Public/Apple/ApplePlatformFile.h
+++ b/Engine/Source/Runtime/Core/Public/Apple/ApplePlatformFile.h
@@ -42,6 +42,7 @@ public:
 	virtual IFileHandle* OpenRead(const TCHAR* Filename, bool bAllowWrite = false) override;
 	virtual IFileHandle* OpenWrite(const TCHAR* Filename, bool bAppend = false, bool bAllowRead = false) override;
 	virtual FOpenMappedResult OpenMappedEx(const TCHAR* Filename, EOpenReadFlags OpenOptions = EOpenReadFlags::None, int64 MaximumSize = 0) override;
+	virtual FOpenMappedResult OpenMappedEx2(const TCHAR* Filename, EOpenMappedFlags OpenOptions = EOpenMappedFlags::None, int64 MaximumSize = 0) override;
 	virtual bool DirectoryExists(const TCHAR* Directory) override;
 	virtual bool CreateDirectory(const TCHAR* Directory) override;
 	virtual bool DeleteDirectory(const TCHAR* Directory) override;
@@ -52,8 +53,11 @@ public:
 	virtual bool IterateDirectoryStat(const TCHAR* Directory, FDirectoryStatVisitor& Visitor) override;
 
 	virtual bool CopyFile(const TCHAR* To, const TCHAR* From, EPlatformFileRead ReadFlags = EPlatformFileRead::None, EPlatformFileWrite WriteFlags = EPlatformFileWrite::None) override;
-private:
+protected:
 	virtual FString ConvertToPlatformPath(const FString& Filename, bool bForWrite, bool bIsPublicWrite) { return Filename; }
+	virtual int32 OpenFileForMemMap(const FString& NormalizedFilename, EOpenMappedFlags OpenOptions, FString& OutFinalPath);
+	int GetModeFlagsForOpenWrite() const;
+private:
 	bool IterateDirectoryCommon(const TCHAR* Directory, const TFunctionRef<bool(struct dirent*)>& Visitor);
 
 	int32 Stat(const TCHAR* Filename, struct stat* OutFileInfo);
diff --git a/Engine/Source/Runtime/Core/Public/Async/MappedFileHandle.h b/Engine/Source/Runtime/Core/Public/Async/MappedFileHandle.h
--- a/Engine/Source/Runtime/Core/Public/Async/MappedFileHandle.h
+++ b/Engine/Source/Runtime/Core/Public/Async/MappedFileHandle.h
@@ -96,7 +96,7 @@ enum class EMappedFileFlags
 {
 	ENone = 0,				// Do nothing
 	EPreloadHint = 1,		// Preload the data. This is only a hint and might be ignored, see IMappedFileRegion::PreloadHint
-	EFileWritable = 2,		// Make the mapped file writable. Requires OpenMappedEx to be called with OpenOptions set to EOpenReadFlags::AllowWrite.
+	EFileWritable = 2,		// Make the mapped file writable. Requires OpenMappedEx2 to be called with OpenOptions set to EOpenMappedFlags::OpenForWrite.
 							// This will create a shared mapping on Unix platform to allow any writes to be flushed to disk
 };
 
diff --git a/Engine/Source/Runtime/Core/Public/GenericPlatform/GenericPlatformFile.h b/Engine/Source/Runtime/Core/Public/GenericPlatform/GenericPlatformFile.h
--- a/Engine/Source/Runtime/Core/Public/GenericPlatform/GenericPlatformFile.h
+++ b/Engine/Source/Runtime/Core/Public/GenericPlatform/GenericPlatformFile.h
@@ -549,7 +549,7 @@ public:
 	{
 		/** No additional flags. */
 		None		= 0,
-		/** Allow other handles/processes to read from this file. */
+		/** Allow the returned handle to be read from. On certain platforms allow other handles/processes to read from this file as well. */
 		AllowRead	= 1 << 0,
 		/** Append to file. */ 
 		Append		= 1 << 1,
@@ -786,16 +786,52 @@ public:
 	{
 	}
 
+	/** Flags to be used when opening a file for via IPlatformFile::OpenMappedEx */
+	enum class EOpenMappedFlags : uint8
+	{
+		None			= 0,
+		/** Allow the handle to be written to via the memory map. */
+		OpenForWrite	= 1 << 0,
+		/** If the file already exists, do not replace it. Only applies if OpenForWrite is set. */
+		OpenExisting	= 1 << 1,
+		/** Allow other handles/processes to write to this file. 
+		  * On certain platforms this flag is needed to open files that are currently being written to as well. 
+		  * Only applies if OpenForWrite is not set.
+		  */
+		AllowWrite		= 1 << 2,
+		/** On certain platforms allow other handles/processes to read from this file as well.
+		  * Only applies if OpenForWrite is set.
+		  */
+		AllowRead		= 1 << 3,
+		/** Allow the file to be deleted or renamed while keeping this handle valid for reading */
+		AllowDelete		= 1 << 4,
+		
+	};
+
 	/** Open a file for async reading. This call does hit the disk; it is synchronous open.
 	*
 	* @param Filename		File to be mapped. This doesn't actually map anything, just opens the file.
-	* @param OpenOptions	Allows specifying map to be writable via AllowWrite or be ReadOnly by deafault.
+	* @param OpenOptions	Allows specifying map to be writable via AllowWrite or be ReadOnly by default.
 	* @return Close the file by deleting the handle. A non-null return value does mean the file exists.
 	* Null can be returned for many reasons even if the file exists. Perhaps this platform does not support mapped files, or this file is compressed in a pak file.
 	* Generally you attempt to open mapped, and if that fails, then use other file operations instead.
+	* The exact file open behavior may vary per-platform when using AllowWrite.
+	* Prefer OpenMappedEx2 which has more consistent behavior across all supported platforms. 
+	* OpenMappedEx is not currently deprecated so as not to change any platform specific behavior that may be implicitly relied upon.
 	*/
 	CORE_API virtual FOpenMappedResult OpenMappedEx(const TCHAR* Filename, EOpenReadFlags OpenOptions = EOpenReadFlags::None, int64 MaximumSize = 0);
 
+	/** Open a file for async reading. This call does hit the disk; it is synchronous open.
+	*
+	* @param Filename		File to be mapped. This doesn't actually map anything, just opens the file.
+	* @param OpenOptions	Allows specifying map to be writable via OpenForWrite or be ReadOnly by default.
+	* @return Close the file by deleting the handle. A non-null return value does mean the file exists.
+	* Null can be returned for many reasons even if the file exists. Perhaps this platform does not support mapped files, or this file is compressed in a pak file.
+	* Generally you attempt to open mapped, and if that fails, then use other file operations instead.
+	* Prefer OpenMappedEx2 which has consistent behavior across all supported platforms.
+	*/
+	CORE_API virtual FOpenMappedResult OpenMappedEx2(const TCHAR* Filename, EOpenMappedFlags OpenOptions = EOpenMappedFlags::None, int64 MaximumSize = 0);
+
 	CORE_API virtual void GetTimeStampPair(const TCHAR* PathA, const TCHAR* PathB, FDateTime& OutTimeStampA, FDateTime& OutTimeStampB);
 
 	/** Return the modification time of a file in the local time of the calling code (GetTimeStamp returns UTC). Returns FDateTime::MinValue() on failure **/
@@ -1079,3 +1115,4 @@ inline FFileJournalData::FFileJournalData()
 
 ENUM_CLASS_FLAGS(IPlatformFile::EOpenReadFlags);
 ENUM_CLASS_FLAGS(IPlatformFile::EOpenWriteFlags);
+ENUM_CLASS_FLAGS(IPlatformFile::EOpenMappedFlags);
diff --git a/Engine/Source/Runtime/Core/Public/HAL/IPlatformFileCachedWrapper.h b/Engine/Source/Runtime/Core/Public/HAL/IPlatformFileCachedWrapper.h
--- a/Engine/Source/Runtime/Core/Public/HAL/IPlatformFileCachedWrapper.h
+++ b/Engine/Source/Runtime/Core/Public/HAL/IPlatformFileCachedWrapper.h
@@ -458,6 +458,10 @@ public:
 	{
 		return LowerLevel->OpenMappedEx(Filename, OpenOptions, MaximumSize);
 	}
+	virtual FOpenMappedResult OpenMappedEx2(const TCHAR* Filename, EOpenMappedFlags OpenOptions = EOpenMappedFlags::None, int64 MaximumSize = 0) override
+	{
+		return LowerLevel->OpenMappedEx2(Filename, OpenOptions, MaximumSize);
+	}
 	virtual void SetAsyncMinimumPriority(EAsyncIOPriorityAndFlags MinPriority) override
 	{
 		LowerLevel->SetAsyncMinimumPriority(MinPriority);
diff --git a/Engine/Source/Runtime/Core/Public/HAL/IPlatformFileLogWrapper.h b/Engine/Source/Runtime/Core/Public/HAL/IPlatformFileLogWrapper.h
--- a/Engine/Source/Runtime/Core/Public/HAL/IPlatformFileLogWrapper.h
+++ b/Engine/Source/Runtime/Core/Public/HAL/IPlatformFileLogWrapper.h
@@ -504,6 +504,7 @@ public:
 		return Result;
 	}
 	virtual FOpenMappedResult OpenMappedEx(const TCHAR* Filename, EOpenReadFlags OpenOptions = EOpenReadFlags::None, int64 MaximumSize = 0) override;
+	virtual FOpenMappedResult OpenMappedEx2(const TCHAR* Filename, EOpenMappedFlags OpenOptions = EOpenMappedFlags::None, int64 MaximumSize = 0) override;
 
 	virtual void SetAsyncMinimumPriority(EAsyncIOPriorityAndFlags Priority) override
 	{
diff --git a/Engine/Source/Runtime/Core/Public/HAL/IPlatformFileOpenLogWrapper.h b/Engine/Source/Runtime/Core/Public/HAL/IPlatformFileOpenLogWrapper.h
--- a/Engine/Source/Runtime/Core/Public/HAL/IPlatformFileOpenLogWrapper.h
+++ b/Engine/Source/Runtime/Core/Public/HAL/IPlatformFileOpenLogWrapper.h
@@ -264,6 +264,7 @@ public:
 		return new FLoggingAsyncReadFileHandle(this, Filename, LowerLevel->OpenAsyncRead(Filename, bAllowWrite));
 	}
 	virtual FOpenMappedResult OpenMappedEx(const TCHAR* Filename, EOpenReadFlags OpenOptions = EOpenReadFlags::None, int64 MaximumSize = 0) override;
+	virtual FOpenMappedResult OpenMappedEx2(const TCHAR* Filename, EOpenMappedFlags OpenOptions = EOpenMappedFlags::None, int64 MaximumSize = 0) override;
 
 	void AddToOpenLog(const TCHAR* Filename)
 	{
diff --git a/Engine/Source/Runtime/Core/Public/IOS/IOSPlatformFile.h b/Engine/Source/Runtime/Core/Public/IOS/IOSPlatformFile.h
--- a/Engine/Source/Runtime/Core/Public/IOS/IOSPlatformFile.h
+++ b/Engine/Source/Runtime/Core/Public/IOS/IOSPlatformFile.h
@@ -57,9 +57,12 @@ public:
 	// path or an empty string if it's not found
 	static FString GetSharedContainerPath(const TCHAR* AppGroupID);
 	
+protected:
+	virtual FString ConvertToPlatformPath(const FString& Filename, bool bForWrite, bool bIsPublicWrite) override;
+	virtual int32 OpenFileForMemMap(const FString& NormalizedFilename, EOpenMappedFlags OpenOptions, FString& OutFinalPath) override;
+
 private:
 	bool IterateDirectoryCommon(const TCHAR* Directory, const TFunctionRef<bool(struct dirent*)>& Visitor);
-	virtual FString ConvertToPlatformPath(const FString& Filename, bool bForWrite, bool bIsPublicWrite) override;
 
 	FString ReadDirectory;
 	FString WriteDirectory;
diff --git a/Engine/Source/Runtime/Core/Public/Unix/UnixPlatformFile.h b/Engine/Source/Runtime/Core/Public/Unix/UnixPlatformFile.h
--- a/Engine/Source/Runtime/Core/Public/Unix/UnixPlatformFile.h
+++ b/Engine/Source/Runtime/Core/Public/Unix/UnixPlatformFile.h
@@ -47,6 +47,7 @@ public:
 	CORE_API virtual IFileHandle* OpenRead(const TCHAR* Filename, bool bAllowWrite = false) override;
 	CORE_API virtual IFileHandle* OpenWrite(const TCHAR* Filename, bool bAppend = false, bool bAllowRead = false) override;
 	CORE_API virtual FOpenMappedResult OpenMappedEx(const TCHAR* Filename, EOpenReadFlags OpenOptions = EOpenReadFlags::None, int64 MaximumSize = 0) override;
+	CORE_API virtual FOpenMappedResult OpenMappedEx2(const TCHAR* Filename, EOpenMappedFlags OpenOptions = EOpenMappedFlags::None, int64 MaximumSize = 0) override;
 
 	CORE_API virtual bool DirectoryExists(const TCHAR* Directory) override;
 	CORE_API virtual bool CreateDirectory(const TCHAR* Directory) override;
diff --git a/Engine/Source/Runtime/PakFile/Private/IPlatformFilePak.cpp b/Engine/Source/Runtime/PakFile/Private/IPlatformFilePak.cpp
--- a/Engine/Source/Runtime/PakFile/Private/IPlatformFilePak.cpp
+++ b/Engine/Source/Runtime/PakFile/Private/IPlatformFilePak.cpp
@@ -4655,7 +4655,9 @@ bool FPakPlatformFile::IsPakFileInstalled(const FString& InFilename)
 }
 #endif //HAS_PLATFORM_PAK_INSTALL_CHECK
 
-FOpenMappedResult FPakPlatformFile::OpenMappedEx(const TCHAR* Filename, EOpenReadFlags OpenOptions, int64 MaximumSize)
+template<typename OpenOptionsType>
+FOpenMappedResult TOpenMappedExImpl(FPakPlatformFile* pThis, const TCHAR* Filename, OpenOptionsType OpenOptions, int64 MaximumSize,
+	FOpenMappedResult(IPlatformFile::*LowLevelOpenMappedEx)(const TCHAR*, OpenOptionsType, int64))
 {
 	if (!GMMIO_Enable)
 	{
@@ -4663,10 +4665,10 @@ FOpenMappedResult FPakPlatformFile::OpenMappedEx(const TCHAR* Filename, EOpenRea
 	}
 
 #if !UE_BUILD_SHIPPING
-	if (bLookLooseFirst && IsNonPakFilenameAllowed(Filename))
+	if (pThis->bLookLooseFirst && pThis->IsNonPakFilenameAllowed(Filename))
 	{
-		FOpenMappedResult Result = LowerLevel->OpenMappedEx(Filename, OpenOptions, MaximumSize);
-		if (!Result.HasError())
+		FOpenMappedResult Result = (pThis->LowerLevel->*LowLevelOpenMappedEx)(Filename, OpenOptions, MaximumSize);
+		if (Result.HasValue())
 		{
 			return MakeValue(Result.StealValue());
 		}
@@ -4676,7 +4678,7 @@ FOpenMappedResult FPakPlatformFile::OpenMappedEx(const TCHAR* Filename, EOpenRea
 	// Check pak files first
 	FPakEntry FileEntry;
 	TRefCountPtr<FPakFile> PakEntry;
-	if (FindFileInPakFiles(Filename, &PakEntry, &FileEntry) && PakEntry.IsValid())
+	if (pThis->FindFileInPakFiles(Filename, &PakEntry, &FileEntry) && PakEntry.IsValid())
 	{
 		if (FileEntry.CompressionMethodIndex != 0 || (FileEntry.Flags & FPakEntry::Flag_Encrypted) != 0)
 		{
@@ -4685,7 +4687,7 @@ FOpenMappedResult FPakPlatformFile::OpenMappedEx(const TCHAR* Filename, EOpenRea
 		FScopeLock Lock(&PakEntry->MappedFileHandleCriticalSection);
 		if (!PakEntry->MappedFileHandle)
 		{
-			FOpenMappedResult Result = LowerLevel->OpenMappedEx(*PakEntry->GetFilename(), OpenOptions, MaximumSize);
+			FOpenMappedResult Result = (pThis->LowerLevel->*LowLevelOpenMappedEx)(*PakEntry->GetFilename(), OpenOptions, MaximumSize);
 			if (Result.HasError())
 			{
 				return MakeError(Result.StealError());
@@ -4696,13 +4698,23 @@ FOpenMappedResult FPakPlatformFile::OpenMappedEx(const TCHAR* Filename, EOpenRea
 
 		return MakeValue(new FMappedFilePakProxy(PakEntry->MappedFileHandle, FileEntry.Offset + FileEntry.GetSerializedSize(PakEntry->GetInfo().Version), FileEntry.UncompressedSize, PakEntry->TotalSize(), Filename));
 	}
-	if (IsNonPakFilenameAllowed(Filename))
+	if (pThis->IsNonPakFilenameAllowed(Filename))
 	{
-		return LowerLevel->OpenMappedEx(Filename, OpenOptions, MaximumSize);
+		return (pThis->LowerLevel->*LowLevelOpenMappedEx)(Filename, OpenOptions, MaximumSize);
 	}
 	return MakeError(FString::Printf(TEXT("Failed to open mapped file '%s'"), Filename));
 }
 
+FOpenMappedResult FPakPlatformFile::OpenMappedEx(const TCHAR* Filename, EOpenReadFlags OpenOptions, int64 MaximumSize)
+{
+	return TOpenMappedExImpl(this, Filename, OpenOptions, MaximumSize, &IPlatformFile::OpenMappedEx);
+}
+
+FOpenMappedResult FPakPlatformFile::OpenMappedEx2(const TCHAR* Filename, EOpenMappedFlags OpenOptions, int64 MaximumSize)
+{
+	return TOpenMappedExImpl(this, Filename, OpenOptions, MaximumSize, &IPlatformFile::OpenMappedEx2);
+}
+
 IAsyncReadFileHandle* FPakPlatformFile::OpenAsyncRead(const TCHAR* Filename, bool bAllowWrite /*= false*/)
 {
 	CSV_SCOPED_TIMING_STAT(FileIOVerbose, PakOpenAsyncRead);
diff --git a/Engine/Source/Runtime/PakFile/Public/IPlatformFilePak.h b/Engine/Source/Runtime/PakFile/Public/IPlatformFilePak.h
--- a/Engine/Source/Runtime/PakFile/Public/IPlatformFilePak.h
+++ b/Engine/Source/Runtime/PakFile/Public/IPlatformFilePak.h
@@ -803,6 +803,9 @@ public:
 
 private:
 	friend class FPakPlatformFile;
+	template<typename OpenOptionsType>
+	friend FOpenMappedResult TOpenMappedExImpl(class FPakPlatformFile* pThis, const TCHAR* Filename, OpenOptionsType OpenOptions, int64 MaximumSize,
+		FOpenMappedResult(IPlatformFile::* LowLevelOpenMappedEx)(const TCHAR*, OpenOptionsType, int64));
 
 	/** Pak filename. */
 	FString PakFilename;
@@ -2521,6 +2524,11 @@ public:
 	PAKFILE_API virtual void SetAsyncMinimumPriority(EAsyncIOPriorityAndFlags Priority) override;
 
 	virtual FOpenMappedResult OpenMappedEx(const TCHAR* Filename, EOpenReadFlags OpenOptions = EOpenReadFlags::None, int64 MaximumSize = 0) override;
+	virtual FOpenMappedResult OpenMappedEx2(const TCHAR* Filename, EOpenMappedFlags OpenOptions = EOpenMappedFlags::None, int64 MaximumSize = 0) override;
+	template<typename OpenOptionsType>
+	friend FOpenMappedResult TOpenMappedExImpl(FPakPlatformFile* pThis, const TCHAR* Filename, OpenOptionsType OpenOptions, int64 MaximumSize,
+		FOpenMappedResult(IPlatformFile::*LowLevelOpenMappedEx)(const TCHAR*, OpenOptionsType, int64));
+
 	/**
 	 * Converts a filename to a path inside pak file.
 	 *
diff --git a/Engine/Source/Runtime/SandboxFile/Private/IPlatformFileSandboxWrapper.cpp b/Engine/Source/Runtime/SandboxFile/Private/IPlatformFileSandboxWrapper.cpp
--- a/Engine/Source/Runtime/SandboxFile/Private/IPlatformFileSandboxWrapper.cpp
+++ b/Engine/Source/Runtime/SandboxFile/Private/IPlatformFileSandboxWrapper.cpp
@@ -996,6 +996,16 @@ FOpenMappedResult FSandboxPlatformFile::OpenMappedEx(const TCHAR* Filename, EOpe
 	return LowerLevel->OpenMappedEx(Filename, OpenOptions, MaximumSize);
 }
 
+FOpenMappedResult FSandboxPlatformFile::OpenMappedEx2(const TCHAR* Filename, EOpenMappedFlags OpenOptions, int64 MaximumSize)
+{
+	FString UserFilename(*ConvertToSandboxPath(Filename));
+	if (!OkForInnerAccess(Filename) || LowerLevel->FileExists(*UserFilename))
+	{
+		return LowerLevel->OpenMappedEx2(*UserFilename, OpenOptions, MaximumSize);
+	}
+	return LowerLevel->OpenMappedEx2(Filename, OpenOptions, MaximumSize);
+}
+
 /**
  * Module for the sandbox file
  */
diff --git a/Engine/Source/Runtime/SandboxFile/Public/IPlatformFileSandboxWrapper.h b/Engine/Source/Runtime/SandboxFile/Public/IPlatformFileSandboxWrapper.h
--- a/Engine/Source/Runtime/SandboxFile/Public/IPlatformFileSandboxWrapper.h
+++ b/Engine/Source/Runtime/SandboxFile/Public/IPlatformFileSandboxWrapper.h
@@ -225,6 +225,7 @@ public:
 	SANDBOXFILE_API virtual void					SetAsyncMinimumPriority(EAsyncIOPriorityAndFlags Priority) override;
 
 	SANDBOXFILE_API virtual FOpenMappedResult		OpenMappedEx(const TCHAR* Filename, EOpenReadFlags OpenOptions = EOpenReadFlags::None, int64 MaximumSize = 0) override;
+	SANDBOXFILE_API virtual FOpenMappedResult		OpenMappedEx2(const TCHAR* Filename, EOpenMappedFlags OpenOptions = EOpenMappedFlags::None, int64 MaximumSize = 0) override;
 
 	friend class FSandboxVisitor;
 	friend class FSandboxStatVisitor;
diff --git a/Engine/Source/Runtime/StorageServerClient/Private/StorageServerPlatformFile.cpp b/Engine/Source/Runtime/StorageServerClient/Private/StorageServerPlatformFile.cpp
--- a/Engine/Source/Runtime/StorageServerClient/Private/StorageServerPlatformFile.cpp
+++ b/Engine/Source/Runtime/StorageServerClient/Private/StorageServerPlatformFile.cpp
@@ -1795,6 +1795,15 @@ FOpenMappedResult FStorageServerPlatformFile::OpenMappedEx(const TCHAR* Filename
 	return MakeError(FString::Printf(TEXT("Can't open mapped file '%s'"), Filename));
 }
 
+FOpenMappedResult FStorageServerPlatformFile::OpenMappedEx2(const TCHAR* Filename, EOpenMappedFlags OpenOptions, int64 MaximumSize)
+{
+	if (LowerLevel && IsNonServerFilenameAllowed(Filename))
+	{
+		return LowerLevel->OpenMappedEx2(Filename, OpenOptions, MaximumSize);
+	}
+	return MakeError(FString::Printf(TEXT("Can't open mapped file '%s'"), Filename));
+}
+
 bool FStorageServerPlatformFile::DirectoryExists(const TCHAR* Directory)
 {
 	TStringBuilder<1024> StorageServerDirectory;
diff --git a/Engine/Source/Runtime/StorageServerClient/Private/StorageServerPlatformFile.h b/Engine/Source/Runtime/StorageServerClient/Private/StorageServerPlatformFile.h
--- a/Engine/Source/Runtime/StorageServerClient/Private/StorageServerPlatformFile.h
+++ b/Engine/Source/Runtime/StorageServerClient/Private/StorageServerPlatformFile.h
@@ -102,6 +102,7 @@ public:
 	virtual bool IterateDirectoryRecursively(const TCHAR* Directory, FDirectoryVisitor& Visitor) override;
 	virtual bool IterateDirectoryStat(const TCHAR* Directory, FDirectoryStatVisitor& Visitor) override;
 	virtual FOpenMappedResult OpenMappedEx(const TCHAR* Filename, EOpenReadFlags OpenOptions = EOpenReadFlags::None, int64 MaximumSize = 0) override;
+	virtual FOpenMappedResult OpenMappedEx2(const TCHAR* Filename, EOpenMappedFlags OpenOptions = EOpenMappedFlags::None, int64 MaximumSize = 0) override;
 	virtual FString GetFilenameOnDisk(const TCHAR* Filename) override;
 	virtual bool DeleteFile(const TCHAR* Filename) override;
 	virtual bool MoveFile(const TCHAR* To, const TCHAR* From) override;
