[Core] Add key storage back into compact set with a define to be able to toggle it on or off
* Natvis improvements and fixed to work with set items that are smaller than 4 bytes.
* Allow garbage reference collection if key storage is enabled on compact sets
* Rename unhash macros to be based on feature than set type

#tests Foundation tests, running game and preflights against previously breaking scenarios
#Synced-CL 51676158

diff --git a/Engine/Source/Runtime/Core/Public/Containers/CompactHashTable.h b/Engine/Source/Runtime/Core/Public/Containers/CompactHashTable.h
--- a/Engine/Source/Runtime/Core/Public/Containers/CompactHashTable.h
+++ b/Engine/Source/Runtime/Core/Public/Containers/CompactHashTable.h
@@ -64,7 +64,7 @@ namespace UE::Core::CompactHashTable
 
 	[[nodiscard]] UE_FORCEINLINE_HINT constexpr size_t GetMemoryRequiredInBytes(uint32 IndexCount, uint32 HashCount)
 	{
-		return ((size_t)IndexCount + HashCount) << GetTypeShift(IndexCount);
+		return (((size_t)IndexCount * UE::Core::Private::CompactSetAllocatorHelpers::KeyStorageMultiplier) + HashCount) << GetTypeShift(IndexCount);
 	}
 
 	[[nodiscard]] UE_FORCEINLINE_HINT constexpr size_t GetMemoryAlignment()
@@ -173,7 +173,7 @@ namespace UE::Core::CompactHashTable
 			*LastLookupEntry = (IndexType)Index;
 
 			// Move the last elements next index into the new slot as well
-					NextIndexData[Index] = NextIndexData[LastIndex];
+			NextIndexData[Index] = NextIndexData[LastIndex];
 		}
 
 		// Return whatever next is at the original lookup entry that we fixed up. If it happened to be the last index the the previous block would have fixed it up for us
@@ -314,11 +314,16 @@ protected:
 class FConstCompactHashTableView
 {
 public:
+	static constexpr uint8 KeyStorageMultiplier = UE::Core::Private::CompactSetAllocatorHelpers::KeyStorageMultiplier;
+
 	explicit FConstCompactHashTableView() = default;
 
 	inline explicit FConstCompactHashTableView(const uint8* Memory, uint32 InNextIndexCount, uint32 InHashCount, size_t MemorySize)
 	: NextIndexData(Memory)
-	, HashData(Memory + (size_t(InNextIndexCount) << UE::Core::CompactHashTable::GetTypeShift(InNextIndexCount)))
+	#if UE_ENABLE_COMPACTHASHTABLE_KEYSTORAGE
+	, KeyStorage(Memory + (size_t(InNextIndexCount) << UE::Core::CompactHashTable::GetTypeShift(InNextIndexCount)))
+	#endif
+	, HashData(Memory + (size_t(InNextIndexCount) << UE::Core::CompactHashTable::GetTypeShift(InNextIndexCount)) * KeyStorageMultiplier)
 	, NextIndexCount(InNextIndexCount)
 	, HashCount(InHashCount)
 	{
@@ -335,7 +340,7 @@ public:
 	// Functions used to search
 	[[nodiscard]] inline uint32 GetFirst(uint32 Key) const
 	{
-		#define UE_COMPACTHASHTABLE_EXECUTEBYTYPE(Type) return UE::Core::CompactHashTable::GetFirst(Key, (const Type *)HashData, HashCount)
+		#define UE_COMPACTHASHTABLE_EXECUTEBYTYPE(Type) return UE::Core::CompactHashTable::GetFirst(Key, (const Type*)HashData, HashCount)
 		UE_COMPACTHASHTABLE_CALLBYTYPE(NextIndexCount);
 		#undef UE_COMPACTHASHTABLE_EXECUTEBYTYPE
 		return INDEX_NONE;
@@ -344,7 +349,7 @@ public:
 	// Advanced used for manual inspection of the hash data
 	[[nodiscard]] inline uint32 GetFirstByIndex(uint32 HashIndex) const
 	{
-		#define UE_COMPACTHASHTABLE_EXECUTEBYTYPE(Type) return UE::Core::CompactHashTable::GetFirstByIndex(HashIndex, (const Type *)HashData, HashCount)
+		#define UE_COMPACTHASHTABLE_EXECUTEBYTYPE(Type) return UE::Core::CompactHashTable::GetFirstByIndex(HashIndex, (const Type*)HashData, HashCount)
 		UE_COMPACTHASHTABLE_CALLBYTYPE(NextIndexCount);
 		#undef UE_COMPACTHASHTABLE_EXECUTEBYTYPE
 		return INDEX_NONE;
@@ -354,7 +359,7 @@ public:
 	{
 		checkSlow(CurrentCount <= NextIndexCount);
 	
-		#define UE_COMPACTHASHTABLE_EXECUTEBYTYPE(Type) return UE::Core::CompactHashTable::GetNext(Index, (const Type *)NextIndexData, CurrentCount)
+		#define UE_COMPACTHASHTABLE_EXECUTEBYTYPE(Type) return UE::Core::CompactHashTable::GetNext(Index, (const Type*)NextIndexData, CurrentCount)
 		UE_COMPACTHASHTABLE_CALLBYTYPE(NextIndexCount);
 		#undef UE_COMPACTHASHTABLE_EXECUTEBYTYPE
 		return INDEX_NONE;
@@ -364,14 +369,32 @@ public:
 	[[nodiscard]] inline uint32 Find(uint32 Key, uint32 CurrentCount, const PredicateType& Predicate) const
 	{
 		checkSlow(CurrentCount <= NextIndexCount);
-		#define UE_COMPACTHASHTABLE_EXECUTEBYTYPE(Type) return UE::Core::CompactHashTable::Find(Key, (const Type *)HashData, HashCount, (const Type *)NextIndexData, CurrentCount, Predicate)
+		#define UE_COMPACTHASHTABLE_EXECUTEBYTYPE(Type) return UE::Core::CompactHashTable::Find(Key, (const Type*)HashData, HashCount, (const Type*)NextIndexData, CurrentCount, Predicate)
+		UE_COMPACTHASHTABLE_CALLBYTYPE(NextIndexCount);
+		#undef UE_COMPACTHASHTABLE_EXECUTEBYTYPE
+		return INDEX_NONE;
+	}
+
+	#if UE_ENABLE_COMPACTHASHTABLE_KEYSTORAGE
+	// Functions used to search
+	[[nodiscard]] inline uint32 GetStoredKey(uint32 Index, uint32 CurrentCount) const
+	{
+		checkSlow(CurrentCount <= NextIndexCount);
+	
+		#define UE_COMPACTHASHTABLE_EXECUTEBYTYPE(Type) return ((const Type*)KeyStorage)[Index];
 		UE_COMPACTHASHTABLE_CALLBYTYPE(NextIndexCount);
 		#undef UE_COMPACTHASHTABLE_EXECUTEBYTYPE
 		return INDEX_NONE;
 	}
+	#endif
 
 protected:
 	const uint8* NextIndexData = nullptr;
+
+	#if UE_ENABLE_COMPACTHASHTABLE_KEYSTORAGE
+	const uint8* KeyStorage = nullptr;
+	#endif
+
 	const uint8* HashData = nullptr;
 	uint32 NextIndexCount = 0;
 	uint32 HashCount = 0;
@@ -397,14 +420,24 @@ public:
 	inline void Add(uint32 Index, uint32 Key) const
 	{
 		checkSlow(Index < NextIndexCount);
+		#if UE_ENABLE_COMPACTHASHTABLE_KEYSTORAGE
+		#define UE_COMPACTHASHTABLE_EXECUTEBYTYPE(Type) UE::Core::CompactHashTable::Add(Index, Key, (Type*)(HashData), HashCount, (Type*)(NextIndexData), NextIndexCount); \
+			((Type*)KeyStorage)[Index] = (Type)(Key & (HashCount - 1))
+		#else
 		#define UE_COMPACTHASHTABLE_EXECUTEBYTYPE(Type) UE::Core::CompactHashTable::Add(Index, Key, (Type*)(HashData), HashCount, (Type*)(NextIndexData), NextIndexCount)
+		#endif
 		UE_COMPACTHASHTABLE_CALLBYTYPE(NextIndexCount);
 		#undef UE_COMPACTHASHTABLE_EXECUTEBYTYPE
 	}
 
 	UE_FORCEINLINE_HINT uint32 Remove(uint32 Index, uint32 Key, uint32 LastIndex, uint32 OptLastKey) const
 	{
+		#if UE_ENABLE_COMPACTHASHTABLE_KEYSTORAGE
+		#define UE_COMPACTHASHTABLE_EXECUTEBYTYPE(Type) ((Type*)KeyStorage)[Index] = ((Type*)KeyStorage)[LastIndex]; \
+			return UE::Core::CompactHashTable::Remove(Index, Key, LastIndex, OptLastKey, (Type*)(HashData), HashCount, (Type*)(NextIndexData), NextIndexCount)
+		#else
 		#define UE_COMPACTHASHTABLE_EXECUTEBYTYPE(Type) return UE::Core::CompactHashTable::Remove(Index, Key, LastIndex, OptLastKey, (Type*)(HashData), HashCount, (Type*)(NextIndexData), NextIndexCount)
+		#endif
 		UE_COMPACTHASHTABLE_CALLBYTYPE(NextIndexCount);
 		#undef UE_COMPACTHASHTABLE_EXECUTEBYTYPE
 		return INDEX_NONE;
@@ -412,7 +445,12 @@ public:
 
 	UE_FORCEINLINE_HINT uint32 RemoveStable(uint32 Index, uint32 Key) const
 	{
+		#if UE_ENABLE_COMPACTHASHTABLE_KEYSTORAGE
+		#define UE_COMPACTHASHTABLE_EXECUTEBYTYPE(Type) FMemory::Memmove(((Type*)KeyStorage) + Index, ((Type*)KeyStorage) + Index + 1, (NextIndexCount - Index - 1) * sizeof(Type)); \
+			return UE::Core::CompactHashTable::RemoveStable(Index, Key, (Type*)(HashData), HashCount, (Type*)(NextIndexData), NextIndexCount)
+		#else
 		#define UE_COMPACTHASHTABLE_EXECUTEBYTYPE(Type) return UE::Core::CompactHashTable::RemoveStable(Index, Key, (Type*)(HashData), HashCount, (Type*)(NextIndexData), NextIndexCount)
+		#endif
 		UE_COMPACTHASHTABLE_CALLBYTYPE(NextIndexCount);
 		#undef UE_COMPACTHASHTABLE_EXECUTEBYTYPE
 		return INDEX_NONE;
diff --git a/Engine/Source/Runtime/Core/Public/Containers/CompactSet.h.inl b/Engine/Source/Runtime/Core/Public/Containers/CompactSet.h.inl
--- a/Engine/Source/Runtime/Core/Public/Containers/CompactSet.h.inl
+++ b/Engine/Source/Runtime/Core/Public/Containers/CompactSet.h.inl
@@ -54,14 +54,15 @@ namespace UE::Core::Private
  *    the cache misses of reading key data and doing the hashing *before* acquiring the lock.
  *
  * Here's a visual example of the data layout for the compact set
- *  ____________ ___________ ________________ ____________
- * |            |           |                |            |
- * | Data Array | Hash Size | Collision List | Hash Table |
- * |____________|___________|________________|____________|
+ *  ____________ ___________ ________________ _____________ ____________
+ * |            |           |                | (Optional)  |            |
+ * | Data Array | Hash Size | Collision List | Key Storage | Hash Table |
+ * |____________|___________|________________|_____________|____________|
  * 
  * Data Array - Payload of the set. This is a just a regular array of items without any empty spots
  * Hash Size - 4 byte integer to reference how big the hash table is. Storing this is significantly faster than trying to recalculate it
  * Collision List - For each entry in the Data Array there is an entry in this list that may contain a valid index to the next item to consider for hash table collisions
+ * (Optional) Key Storage - If enabled, stores "GetHashKey(Item) % (HashSize - 1)" so we can remove an item without having to rehash it.
  * Hash Table - Power of 2 table to lookup the first index of a given hash value
  *
  **/
@@ -528,7 +529,16 @@ public:
 		{
 			FCompactHashTableView HashTable = GetHashTableView();
 
-			if (this->NumElements < (SizeType)(HashTable.GetHashCount() / UE_COMPACTSET_UNHASH_THRESHOLD))
+			#if UE_ENABLE_COMPACTHASHTABLE_KEYSTORAGE
+			if (this->NumElements < (SizeType)(HashTable.GetHashCount() / UE_SET_UNHASH_THRESHOLD_STORED_KEY))
+			{
+				for (int32 Index = 0; Index < this->NumElements; ++Index)
+				{
+					HashTable.RemoveHashElement(HashTable.GetStoredKey(Index, this->NumElements));
+				}
+			}
+			#else
+			if (this->NumElements < (SizeType)(HashTable.GetHashCount() / UE_SET_UNHASH_THRESHOLD_CALCULATE_KEY))
 			{
 				// if we've only used a portion of the full allocated size then it can be faster
 				// to only clear out the used hash elements instead of memsetting the whole of the memory
@@ -542,6 +552,7 @@ public:
 					HashTable.RemoveHashElement(KeyHash);
 				}
 			}
+			#endif
 			else
 			{
 				HashTable.Reset();
@@ -1537,7 +1548,7 @@ private:
 		// Inlining this can save up to 40% perf compared to GetHashTableView().Find()
 		#define UE_COMPACTHASHTABLE_EXECUTEBYTYPE(Type) \
 				const Type* NextIndices = (const Type *)NextIndicesData; \
-				const Type* HashIndices = NextIndices + this->MaxElements; \
+				const Type* HashIndices = NextIndices + this->MaxElements * this->KeyStorageMultiplier; \
 				for (Type Index = HashIndices[KeyHash & (HashCount-1)]; Index != (Type)INDEX_NONE; Index = NextIndices[Index]) \
 				{ \
 					checkSlow((SizeType)Index < this->NumElements); \
@@ -1561,7 +1572,11 @@ private:
 		ElementType* ElementsData = GetData();
 		FCompactHashTableView HashTable = GetHashTableView();
 
+		#if UE_ENABLE_COMPACTHASHTABLE_KEYSTORAGE
+		RemoveByIndexAndHash<IsStable>(ElementIndex, HashTable.GetStoredKey(ElementIndex, this->NumElements), ElementsData, HashTable);
+		#else
 		RemoveByIndexAndHash<IsStable>(ElementIndex, KeyFuncs::GetKeyHash(KeyFuncs::GetSetKey(ElementsData[ElementIndex])), ElementsData, HashTable);
+		#endif
 	}
 
 	template <bool IsStable = false>
@@ -1586,7 +1601,11 @@ private:
 		}
 		else
 		{
+			#if UE_ENABLE_COMPACTHASHTABLE_KEYSTORAGE
+			const uint32 LastElementHash = HashTable.GetStoredKey(LastElementIndex, this->NumElements);
+			#else
 			const uint32 LastElementHash = KeyFuncs::GetKeyHash(KeyFuncs::GetSetKey(ElementsData[LastElementIndex]));
+			#endif
 			NextIndex = HashTable.Remove(ElementIndex, KeyHash, LastElementIndex, LastElementHash);
 			
 			MoveByRelocate(ElementsData[ElementIndex], ElementsData[LastElementIndex]);
diff --git a/Engine/Source/Runtime/Core/Public/Containers/CompactSetBase.h b/Engine/Source/Runtime/Core/Public/Containers/CompactSetBase.h
--- a/Engine/Source/Runtime/Core/Public/Containers/CompactSetBase.h
+++ b/Engine/Source/Runtime/Core/Public/Containers/CompactSetBase.h
@@ -73,6 +73,8 @@ public:
 protected:
 	using HashCountType = uint32;
 	static constexpr size_t HashCountSize = sizeof(HashCountType);
+	static constexpr uint8 KeyStorageMultiplier = UE::Core::Private::CompactSetAllocatorHelpers::KeyStorageMultiplier;
+
 	using ElementAllocatorType = typename AllocatorType::template ForElementType<uint8>;
 
 	static_assert(std::is_same_v<SizeType, int32>, "TCompactSet currently only supports 32-bit allocators");
@@ -144,8 +146,8 @@ protected:
 		const uint32 TypeSize = UE::Core::CompactHashTable::GetTypeSize(MinElementCount);
 		const uint32 TypeShift = UE::Core::CompactHashTable::GetTypeShift(MinElementCount);
 		const size_t AvailableBytes = TotalBytes - sizeof(HashCountType) - ((size_t)HashCount << TypeShift); // Remove hashtable and HashCount data
-		const size_t MaxElements = AvailableBytes / (Layout.Size + TypeSize); // Calculate the max available ignoring the fact that the hash data has to be aligned
-		const size_t RealAvailableBytes = AlignDown(AvailableBytes - (MaxElements << TypeShift), UE::Core::CompactHashTable::GetMemoryAlignment()); // Remove the max required indexes and align down
+		const size_t MaxElements = AvailableBytes / (Layout.Size + TypeSize * KeyStorageMultiplier); // Calculate the max available ignoring the fact that the hash data has to be aligned
+		const size_t RealAvailableBytes = AlignDown(AvailableBytes - (MaxElements << TypeShift) * KeyStorageMultiplier, UE::Core::CompactHashTable::GetMemoryAlignment()); // Remove the max required indexes and align down
 		return (SizeType)FMath::Min<size_t>(MaxElements, RealAvailableBytes / Layout.Size); // Now we can get the true number of elements we could potentially use within the aligned space
 	}
 
@@ -238,12 +240,22 @@ protected:
 				const uint8* OldHashTable = Elements.GetAllocation() + GetElementsSizeInBytes(MaxElements, Layout);
 				uint8* NewHashTable = Elements.GetAllocation() + GetElementsSizeInBytes(NewMaxElements, Layout);
 
-				const uint8* OldHashLocation = OldHashTable + ((size_t)MaxElements << NewTypeShift) + HashCountSize;
-				uint8* NewHashLocation = NewHashTable + ((size_t)NewMaxElements << NewTypeShift) + HashCountSize;
+				const uint8* OldHashLocation = OldHashTable + ((size_t)MaxElements << NewTypeShift) * KeyStorageMultiplier + HashCountSize;
+				uint8* NewHashLocation = NewHashTable + ((size_t)NewMaxElements << NewTypeShift) * KeyStorageMultiplier + HashCountSize;
+
+				// WARNING: Order matters here since we're moving data around in the same buffer. Make sure we're moving the tail data first
+				// since we're expanding in size. So HashTable, then (optional) key storage, then hashcount and next indexes
 
 				// Copy hash lookup table (first so to free up the space for the other data to move)
 				FMemory::Memmove(NewHashLocation, OldHashLocation, ((size_t)NewHashCount << NewTypeShift));
 
+#if UE_ENABLE_COMPACTHASHTABLE_KEYSTORAGE
+				// Copy key storage if enabled
+				const uint8* OldKeyStorage = OldHashTable + ((size_t)MaxElements << NewTypeShift) + HashCountSize;
+				uint8* NewKeyStorage = NewHashTable + ((size_t)NewMaxElements << NewTypeShift) + HashCountSize;
+				FMemory::Memmove(NewKeyStorage, OldKeyStorage, ((size_t)NumElements << NewTypeShift));
+#endif
+
 				// Copy hash size + next index table
 				FMemory::Memmove(NewHashTable, OldHashTable, ((size_t)NumElements << NewTypeShift) + HashCountSize);
 			}
diff --git a/Engine/Source/Runtime/Core/Public/Containers/ContainerAllocationPolicies.h b/Engine/Source/Runtime/Core/Public/Containers/ContainerAllocationPolicies.h
--- a/Engine/Source/Runtime/Core/Public/Containers/ContainerAllocationPolicies.h
+++ b/Engine/Source/Runtime/Core/Public/Containers/ContainerAllocationPolicies.h
@@ -1492,14 +1492,23 @@ public:
 // Compact set allocation definitions - It operates entirely with uint8 because it includes its lookup table inline with the overall allocation
 //
 
+// If enabled store a copy of the 'key' (key % hashcount - 1, so lookup index) for quick removal
+// This takes more space but allows compact sets to be used with garbage reference removal
+#if !defined UE_ENABLE_COMPACTHASHTABLE_KEYSTORAGE
+#define UE_ENABLE_COMPACTHASHTABLE_KEYSTORAGE 1
+#endif
+
 namespace UE::Core::Private::CompactSetAllocatorHelpers
 {
+	static constexpr uint8 KeyStorageMultiplier = UE_ENABLE_COMPACTHASHTABLE_KEYSTORAGE ? 2 : 1;
+
 	template <uint32 NumInlineElements, int32 ElementSize>
 	constexpr int32 CalculateRequiredBytes()
 	{
 		constexpr uint32 TypeSize = 1 + (NumInlineElements > 0xff) + (NumInlineElements > 0xffff) * 2;
 		constexpr uint32 HashSize = NumInlineElements < 8 ? 4 : FGenericPlatformMath::RoundUpToPowerOfTwo((NumInlineElements / 2) + 1);
-		return Align(NumInlineElements*ElementSize, 4) + 4 + (NumInlineElements + HashSize) * TypeSize;
+		constexpr uint32 IndexStorageSize = NumInlineElements * KeyStorageMultiplier;
+		return Align(NumInlineElements*ElementSize, 4) + 4 + (IndexStorageSize + HashSize) * TypeSize;
 	}
 }
 
diff --git a/Engine/Source/Runtime/Core/Public/Containers/SetUtilities.h b/Engine/Source/Runtime/Core/Public/Containers/SetUtilities.h
--- a/Engine/Source/Runtime/Core/Public/Containers/SetUtilities.h
+++ b/Engine/Source/Runtime/Core/Public/Containers/SetUtilities.h
@@ -11,8 +11,8 @@
 
 // These divisors will determine at what point it's more optimal to individually remove elements from the hash table vs just using memset to reset the whole thing at once
 // Sparse set stores the hashtable index so it's fairly cheap to remove vs compact set that has to recalculate the hash key which is more expensive
-#define UE_SPARSESET_UNHASH_THRESHOLD 4
-#define UE_COMPACTSET_UNHASH_THRESHOLD 32
+#define UE_SET_UNHASH_THRESHOLD_STORED_KEY 4
+#define UE_SET_UNHASH_THRESHOLD_CALCULATE_KEY 32
 
 /**
  * Traits class which determines whether or not a type is a TSet.
diff --git a/Engine/Source/Runtime/Core/Public/Containers/SparseSet.h.inl b/Engine/Source/Runtime/Core/Public/Containers/SparseSet.h.inl
--- a/Engine/Source/Runtime/Core/Public/Containers/SparseSet.h.inl
+++ b/Engine/Source/Runtime/Core/Public/Containers/SparseSet.h.inl
@@ -1393,7 +1393,7 @@ private:
 		FSetElementId* HashPtr = Hash.GetAllocation();
 
 		// Check if it should be faster to clear the hash by going through elements instead of resetting the whole hash
-		if (Num() < (HashSize / UE_SPARSESET_UNHASH_THRESHOLD))
+		if (Num() < (HashSize / UE_SET_UNHASH_THRESHOLD_STORED_KEY))
 		{
 			// Faster path: only reset hash buckets to invalid for elements in the hash
 			for (const SetElementType& Element: Elements)
diff --git a/Engine/Source/Runtime/Core/Tests/Containers/CompactHashTableTests.cpp b/Engine/Source/Runtime/Core/Tests/Containers/CompactHashTableTests.cpp
--- a/Engine/Source/Runtime/Core/Tests/Containers/CompactHashTableTests.cpp
+++ b/Engine/Source/Runtime/Core/Tests/Containers/CompactHashTableTests.cpp
@@ -201,7 +201,8 @@ TEST_CASE_NAMED(FCompactHashTableViewTest, "System::Core::Containers::CompactHas
 			using TestType = FCompactHashTableView;
 
 			constexpr size_t RequiredSize = UE::Core::CompactHashTable::GetMemoryRequiredInBytes(Size, HashSize);
-			static_assert(RequiredSize == sizeof(typename TStaticCompactHashTable<Size>::IndexType) * (Size + HashSize));
+			constexpr uint32 IndexStorageSize = Size * UE::Core::Private::CompactSetAllocatorHelpers::KeyStorageMultiplier;
+			static_assert(RequiredSize == sizeof(typename TStaticCompactHashTable<Size>::IndexType) * (IndexStorageSize + HashSize));
 
 			uint8 TestMemory[RequiredSize];
 			FMemory::Memset(TestMemory, Sentinel, RequiredSize);
@@ -217,7 +218,7 @@ TEST_CASE_NAMED(FCompactHashTableViewTest, "System::Core::Containers::CompactHas
 			
 			for (int32 Index = 0; Index < RequiredSize; ++Index)
 			{
-				if (Index < Size * sizeof(IndexType))
+				if (Index < IndexStorageSize * sizeof(IndexType))
 				{
 					CHECK(TestMemory[Index] == Sentinel);
 				}
diff --git a/Engine/Source/Runtime/Core/Tests/Containers/SetTests.cpp b/Engine/Source/Runtime/Core/Tests/Containers/SetTests.cpp
--- a/Engine/Source/Runtime/Core/Tests/Containers/SetTests.cpp
+++ b/Engine/Source/Runtime/Core/Tests/Containers/SetTests.cpp
@@ -347,7 +347,7 @@ struct TSetTestHelper
 
 			const uint32 PartialElements = 4;
 			// + 1 so it's immune to < vs <=, we want to make sure to test both unhash and memset scenarios for both set types
-			const uint32 MaxElements = (PartialElements + 1) * FMath::Max(UE_SPARSESET_UNHASH_THRESHOLD, UE_COMPACTSET_UNHASH_THRESHOLD);
+			const uint32 MaxElements = (PartialElements + 1) * FMath::Max(UE_SET_UNHASH_THRESHOLD_STORED_KEY, UE_SET_UNHASH_THRESHOLD_CALCULATE_KEY);
 			Set.Empty(MaxElements);
 			CHECK(Set.Max() >= MaxElements);
 			const uint32 CalculatedMaxElements = Set.Max();
diff --git a/Engine/Source/Runtime/CoreUObject/Private/UObject/ObjectBaseUtility.cpp b/Engine/Source/Runtime/CoreUObject/Private/UObject/ObjectBaseUtility.cpp
--- a/Engine/Source/Runtime/CoreUObject/Private/UObject/ObjectBaseUtility.cpp
+++ b/Engine/Source/Runtime/CoreUObject/Private/UObject/ObjectBaseUtility.cpp
@@ -165,9 +165,10 @@ void FScopeCycleCounterUObject::StartObjectTrace(const UObjectBaseUtility* Objec
 }
 #endif
 
+#define UE_GARBAGE_ELIMINATION_ALLOWED !UE_USE_COMPACT_SET_AS_DEFAULT || UE_ENABLE_COMPACTHASHTABLE_KEYSTORAGE
 
 // Console variable so that GarbageCollectorSettings work in the editor but we don't want to use it in runtime as we can't support changing its value from console
-int32 GGarbageEliminationEnabled = !UE_USE_COMPACT_SET_AS_DEFAULT;
+int32 GGarbageEliminationEnabled = UE_GARBAGE_ELIMINATION_ALLOWED;
 static FAutoConsoleVariableRef CVarGarbageEliminationEnabled(
 	TEXT("gc.GarbageEliminationEnabled"),
 	GGarbageEliminationEnabled,
@@ -207,7 +208,7 @@ void InitGarbageElimination()
 
 void UObjectBaseUtility::SetGarbageEliminationEnabled(bool bEnabled)
 {
-#if !UE_USE_COMPACT_SET_AS_DEFAULT
+#if UE_GARBAGE_ELIMINATION_ALLOWED
 	// Try to sync even though we're not gonna use the console var
 	UObjectBaseUtility::bGarbageEliminationEnabled = bEnabled;
 	GGarbageEliminationEnabled = bEnabled;
