Fix Preview crash caused by copying over Flags when Platforms Cvar copy

#Synced-CL 51815809

diff --git a/Engine/Source/Runtime/Core/Private/HAL/ConsoleManager.cpp b/Engine/Source/Runtime/Core/Private/HAL/ConsoleManager.cpp
--- a/Engine/Source/Runtime/Core/Private/HAL/ConsoleManager.cpp
+++ b/Engine/Source/Runtime/Core/Private/HAL/ConsoleManager.cpp
@@ -1616,7 +1616,7 @@ FConsoleVariable<T>* FindOrCreateTypedPlatformCVar(FConsoleVariableExtendedData<
 	TSharedPtr<IConsoleVariable> PlatformCVar = CVar->PlatformValues.FindRef(PlatformKey);
 	if (!PlatformCVar.IsValid())
 	{
-		PlatformCVar = TSharedPtr<IConsoleVariable>(new FConsoleVariable(CVar->GetDefaultTypedValue(), TEXT("Platform CVar copy"), CVar->GetFlags(), CVar));
+		PlatformCVar = TSharedPtr<IConsoleVariable>(new FConsoleVariable(CVar->GetDefaultTypedValue(), TEXT("Platform CVar copy"), (EConsoleVariableFlags)(CVar->GetFlags() & ~ECVF_SetByMask), CVar));
 		
 		// cache it
 		CVar->PlatformValues.Add(PlatformKey, PlatformCVar);
