LinearAllocator: Fix compile error in TSAN configuration: error: conditional expression is ambiguous; 'FMalloc *' can be converted to 'TTSANSafeValue<class FMalloc *>' and vice versa.
#rnx
#Synced-CL 51748961

diff --git a/Engine/Source/Runtime/Core/Private/Memory/LinearAllocator.cpp b/Engine/Source/Runtime/Core/Private/Memory/LinearAllocator.cpp
--- a/Engine/Source/Runtime/Core/Private/Memory/LinearAllocator.cpp
+++ b/Engine/Source/Runtime/Core/Private/Memory/LinearAllocator.cpp
@@ -135,10 +135,14 @@ void* FLinearBlockAllocator::FBlock::GetBlockStart()
 
 FLinearBlockAllocator::FLinearBlockAllocator(FMalloc* InInnerMalloc, uint32 InBlockSize,
 	ESPMode InThreadMode)
-	: InnerMalloc(InInnerMalloc ? InInnerMalloc : UE::Private::GMalloc)
+	: InnerMalloc(InInnerMalloc)
 	, DefaultBlockSize(InBlockSize != 0 ? InBlockSize : UE::LinearAllocator::DefaultBlockSize)
 	, ThreadMode(InThreadMode)
 {
+	if (!InnerMalloc)
+	{
+		InnerMalloc = UE::Private::GMalloc;
+	}
 	static_assert(IsAligned(MinBlockSize, UE::LinearAllocator::BlockAlignment));
 	DefaultBlockSize = FMath::Max(DefaultBlockSize, MinBlockSize);
 	DefaultBlockSize = FMath::RoundUpToPowerOfTwo(DefaultBlockSize);
