Added FUtf8Iterator for decoding UTF-8 views as Unicode codepoints,
Added unit tests.

#Synced-CL 51369557

diff --git a/Engine/Source/Runtime/Core/Private/Utf8/Utf8Iterator.cpp b/Engine/Source/Runtime/Core/Private/Utf8/Utf8Iterator.cpp
new file mode 100644
--- /dev/null
+++ b/Engine/Source/Runtime/Core/Private/Utf8/Utf8Iterator.cpp
@@ -0,0 +1,255 @@
+// Copyright Epic Games, Inc. All Rights Reserved.
+
+#include "Utf8/Utf8Iterator.h"
+
+#include "Containers/StringView.h"
+
+// ---------------------------------------------------------------------------
+// Internal helpers
+// ---------------------------------------------------------------------------
+
+namespace
+{
+	/** Returns true when Octet has the bit pattern 10xxxxxx (a UTF-8 continuation byte). */
+	UE_REWRITE bool IsContinuationByte(uint32 Octet)
+	{
+		return (Octet & 0xC0u) == 0x80u;
+	}
+
+	/** Returns true when Codepoint is in the UTF-16 surrogate range [U+D800, U+DFFF]. */
+	UE_REWRITE bool IsSurrogate(uint32 Codepoint)
+	{
+		return (Codepoint & 0xFFFFF800u) == 0xD800u;
+	}
+} // anonymous namespace
+
+// ---------------------------------------------------------------------------
+// FUtf8Iterator
+// ---------------------------------------------------------------------------
+
+FUtf8Iterator::FUtf8Iterator(FUtf8StringView InView)
+	: Pos(InView.GetData())
+	, ViewEnd(InView.GetData() + InView.Len())
+	, CachedCodepoint(ReplacementCharacter)
+{
+	if (Pos != ViewEnd)
+	{
+		DecodeNext();
+	}
+	else
+	{
+		ViewEnd = nullptr;
+	}
+}
+
+void FUtf8Iterator::operator++()
+{
+	checkfSlow(*this, TEXT("Iterator has already reached the end"));
+	if (Pos != ViewEnd)
+	{
+		DecodeNext();
+	}
+	else
+	{
+		ViewEnd = nullptr;
+	}
+}
+
+void FUtf8Iterator::DecodeNext()
+{
+	// Precondition: Pos < ViewEnd
+	const uint32 Lead = (uint32)(uint8)*Pos;
+
+	// -----------------------------------------------------------------------
+	// Single-byte sequence (U+0000 – U+007F)
+	// -----------------------------------------------------------------------
+	if (Lead < 0x80u)
+	{
+		CachedCodepoint = (UTF32CHAR)Lead;
+		++Pos;
+		return;
+	}
+
+	// -----------------------------------------------------------------------
+	// Continuation byte in lead position, or 0xC0/0xC1 which can never start
+	// a valid sequence (all their encodings would be overlong) — consume just
+	// this byte.
+	// -----------------------------------------------------------------------
+	if (Lead < 0xC2u)
+	{
+		CachedCodepoint = ReplacementCharacter;
+		++Pos;
+		return;
+	}
+
+	// -----------------------------------------------------------------------
+	// Two-byte sequence (U+0080 – U+07FF)
+	// -----------------------------------------------------------------------
+	if (Lead < 0xE0u)
+	{
+		if (ViewEnd - Pos < 2)
+		{
+			// Truncated at end of view
+			CachedCodepoint = ReplacementCharacter;
+			Pos = ViewEnd;
+			return;
+		}
+
+		const uint32 Cont1 = (uint32)(uint8)Pos[1];
+		if (!IsContinuationByte(Cont1))
+		{
+			// Not a continuation byte; do not consume it
+			CachedCodepoint = ReplacementCharacter;
+			Pos += 1;
+			return;
+		}
+
+		const uint32 Codepoint = ((Lead & 0x1Fu) << 6u) | (Cont1 & 0x3Fu);
+		// For Lead >= 0xC2 the minimum codepoint is U+0080, so the overlong check
+		// below is unreachable; it is kept as a safety net.
+		CachedCodepoint = (Codepoint < 0x80u) ? ReplacementCharacter : (UTF32CHAR)Codepoint;
+		Pos += 2;
+		return;
+	}
+
+	// -----------------------------------------------------------------------
+	// Three-byte sequence (U+0800 – U+FFFF, excluding surrogates)
+	// -----------------------------------------------------------------------
+	if (Lead < 0xF0u)
+	{
+		if (ViewEnd - Pos < 2)
+		{
+			CachedCodepoint = ReplacementCharacter;
+			Pos = ViewEnd;
+			return;
+		}
+
+		const uint32 Cont1 = (uint32)(uint8)Pos[1];
+		if (!IsContinuationByte(Cont1))
+		{
+			CachedCodepoint = ReplacementCharacter;
+			Pos += 1;
+			return;
+		}
+
+		// Structural range check on the first continuation byte:
+		//   0xE0 requires Cont1 in 0xA0–0xBF to avoid overlong sequences.
+		//   0xED requires Cont1 in 0x80–0x9F to avoid surrogate codepoints.
+		// When the check fails, consume only the lead byte; Cont1 will be decoded
+		// as a bare continuation (→ replacement) on the next step.
+		if ((Lead == 0xE0u && Cont1 < 0xA0u) || (Lead == 0xEDu && Cont1 >= 0xA0u))
+		{
+			CachedCodepoint = ReplacementCharacter;
+			Pos += 1;
+			return;
+		}
+
+		if (ViewEnd - Pos < 3)
+		{
+			CachedCodepoint = ReplacementCharacter;
+			Pos = ViewEnd;
+			return;
+		}
+
+		const uint32 Cont2 = (uint32)(uint8)Pos[2];
+		if (!IsContinuationByte(Cont2))
+		{
+			CachedCodepoint = ReplacementCharacter;
+			Pos += 2;
+			return;
+		}
+
+		const uint32 Codepoint = ((Lead & 0x0Fu) << 12u) | ((Cont1 & 0x3Fu) << 6u) | (Cont2 & 0x3Fu);
+		// The structural range checks above make the overlong and surrogate cases
+		// unreachable; they are kept as a safety net.
+		CachedCodepoint = (Codepoint < 0x800u || IsSurrogate(Codepoint)) ? ReplacementCharacter : (UTF32CHAR)Codepoint;
+		Pos += 3;
+		return;
+	}
+
+	// -----------------------------------------------------------------------
+	// Four-byte sequence (U+10000 – U+10FFFF)
+	// -----------------------------------------------------------------------
+	if (Lead < 0xF8u)
+	{
+		// 0xF5–0xF7 always produce codepoints above U+10FFFF; consume only the
+		// lead byte so that any following bytes are decoded independently.
+		if (Lead >= 0xF5u)
+		{
+			CachedCodepoint = ReplacementCharacter;
+			++Pos;
+			return;
+		}
+
+		if (ViewEnd - Pos < 2)
+		{
+			CachedCodepoint = ReplacementCharacter;
+			Pos = ViewEnd;
+			return;
+		}
+
+		const uint32 Cont1 = (uint32)(uint8)Pos[1];
+		if (!IsContinuationByte(Cont1))
+		{
+			CachedCodepoint = ReplacementCharacter;
+			Pos += 1;
+			return;
+		}
+
+		// Structural range check on the first continuation byte:
+		//   0xF0 requires Cont1 in 0x90–0xBF to avoid overlong sequences.
+		//   0xF4 requires Cont1 in 0x80–0x8F to avoid codepoints above U+10FFFF.
+		// When the check fails, consume only the lead byte.
+		if ((Lead == 0xF0u && Cont1 < 0x90u) || (Lead == 0xF4u && Cont1 >= 0x90u))
+		{
+			CachedCodepoint = ReplacementCharacter;
+			Pos += 1;
+			return;
+		}
+
+		if (ViewEnd - Pos < 3)
+		{
+			CachedCodepoint = ReplacementCharacter;
+			Pos = ViewEnd;
+			return;
+		}
+
+		const uint32 Cont2 = (uint32)(uint8)Pos[2];
+		if (!IsContinuationByte(Cont2))
+		{
+			CachedCodepoint = ReplacementCharacter;
+			Pos += 2;
+			return;
+		}
+
+		if (ViewEnd - Pos < 4)
+		{
+			CachedCodepoint = ReplacementCharacter;
+			Pos = ViewEnd;
+			return;
+		}
+
+		const uint32 Cont3 = (uint32)(uint8)Pos[3];
+		if (!IsContinuationByte(Cont3))
+		{
+			CachedCodepoint = ReplacementCharacter;
+			Pos += 3;
+			return;
+		}
+
+		const uint32 Codepoint = ((Lead & 0x07u) << 18u) | ((Cont1 & 0x3Fu) << 12u) | ((Cont2 & 0x3Fu) << 6u) | (Cont3 & 0x3Fu);
+		// The structural range checks above make the overlong and out-of-range cases
+		// unreachable; they are kept as a safety net.
+		CachedCodepoint = (Codepoint < 0x10000u || Codepoint > 0x10FFFFu) ? ReplacementCharacter : (UTF32CHAR)Codepoint;
+		Pos += 4;
+		return;
+	}
+
+	// -----------------------------------------------------------------------
+	// Five-byte and six-byte sequences — illegal per RFC 3629.
+	// Also catches 0xFE and 0xFF which are never valid UTF-8 bytes.
+	// Consume only the lead byte; any following bytes are decoded independently.
+	// -----------------------------------------------------------------------
+	CachedCodepoint = ReplacementCharacter;
+	++Pos;
+}
diff --git a/Engine/Source/Runtime/Core/Public/Utf8/Utf8Iterator.h b/Engine/Source/Runtime/Core/Public/Utf8/Utf8Iterator.h
new file mode 100644
--- /dev/null
+++ b/Engine/Source/Runtime/Core/Public/Utf8/Utf8Iterator.h
@@ -0,0 +1,75 @@
+// Copyright Epic Games, Inc. All Rights Reserved.
+
+#pragma once
+
+#include "Containers/StringFwd.h"
+#include "CoreTypes.h"
+
+/**
+ * Iterator for decoding a UTF-8 encoded string view as Unicode codepoints (UTF-32).
+ *
+ * Usage:
+ *   for (FUtf8Iterator It(MyStringView); It; ++It)
+ *   {
+ *       UTF32CHAR Codepoint = *It;
+ *   }
+ *
+ * Invalid UTF-8 sequences are decoded following the Unicode Standard's "maximal subpart"
+ * rule: each maximal initial subsequence of bytes that could be the beginning of a
+ * well-formed sequence is replaced by exactly one U+FFFD replacement character.  In
+ * practice this means that a lone invalid or out-of-range lead byte is consumed by
+ * itself, and any continuation bytes (0x80–0xBF) that follow it are each decoded as
+ * independent errors, each producing their own replacement character.
+ *
+ * Sequences truncated at the end of the view also produce a replacement character and
+ * advance the iterator to the end of the view.
+ *
+ * Zero bytes within the view are treated as ordinary data and do not terminate iteration.
+ *
+ * The iterator is non-copyable. Construct it with an explicit FUtf8StringView.
+ */
+class FUtf8Iterator
+{
+public:
+	UE_NONCOPYABLE(FUtf8Iterator)
+
+	/** The Unicode replacement character emitted for any invalid or truncated sequence. */
+	static constexpr UTF32CHAR ReplacementCharacter = (UTF32CHAR)0xFFFDu;
+
+	/** Construct an iterator positioned at the start of the given UTF-8 view. */
+	[[nodiscard]] CORE_API explicit FUtf8Iterator(FUtf8StringView InView);
+
+	/**
+	 * Returns true while the iterator is positioned at a valid codepoint.
+	 * Becomes false once all code unit sequences in the view have been consumed.
+	 */
+	[[nodiscard]] UE_REWRITE explicit operator bool() const
+	{
+		return ViewEnd != nullptr;
+	}
+
+	/**
+	 * Returns the codepoint at the current position.
+	 * Only valid when operator bool() returns true.
+	 */
+	[[nodiscard]] UE_REWRITE UTF32CHAR operator*() const
+	{
+		return CachedCodepoint;
+	}
+
+	/** Advance to the next codepoint. */
+	CORE_API void operator++();
+	UE_REWRITE void operator++(int)
+	{
+		++(*this);
+	}
+
+private:
+	/** Decode the sequence at Pos, store the result in CachedCodepoint, and
+	 *  advance Pos to the first byte of the following sequence. */
+	void DecodeNext();
+
+	const UTF8CHAR* Pos;
+	const UTF8CHAR* ViewEnd;
+	UTF32CHAR       CachedCodepoint;
+};
diff --git a/Engine/Source/Runtime/Core/Tests/Utf8/Utf8IteratorTest.cpp b/Engine/Source/Runtime/Core/Tests/Utf8/Utf8IteratorTest.cpp
new file mode 100644
--- /dev/null
+++ b/Engine/Source/Runtime/Core/Tests/Utf8/Utf8IteratorTest.cpp
@@ -0,0 +1,358 @@
+// Copyright Epic Games, Inc. All Rights Reserved.
+
+#if WITH_TESTS
+
+#include "Utf8/Utf8Iterator.h"
+
+#include "Containers/Array.h"
+#include "Containers/StringView.h"
+#include "Tests/TestHarnessAdapter.h"
+
+// ---------------------------------------------------------------------------
+// Helpers
+// ---------------------------------------------------------------------------
+
+namespace
+{
+	/** Collect all codepoints produced by iterating over a raw byte span. */
+	TArray<UTF32CHAR> Decode(FUtf8StringView View)
+	{
+		TArray<UTF32CHAR> Result;
+		for (FUtf8Iterator It(View); It; ++It)
+		{
+			Result.Add(*It);
+		}
+		return Result;
+	}
+
+	/** Collect all codepoints from a UTF-8 string literal. */
+	UE_REWRITE TArray<UTF32CHAR> Decode(const UTF8CHAR* Str, int32 Len)
+	{
+		return Decode(FUtf8StringView(Str, Len));
+	}
+
+	/** Collect all codepoints from a UTF-8 string literal. */
+	UE_REWRITE TArray<UTF32CHAR> Decode(const ANSICHAR* Str)
+	{
+		return Decode(FUtf8StringView(Str));
+	}
+}
+
+// ---------------------------------------------------------------------------
+// Tests
+// ---------------------------------------------------------------------------
+
+TEST_CASE_NAMED(FUtf8IteratorTest, "System::Core::Utf8::Utf8Iterator", "[Core][Utf8][SmokeFilter]")
+{
+	// -----------------------------------------------------------------------
+	SECTION("EmptyView")
+	{
+		FUtf8Iterator It(FUtf8StringView{});
+		CHECK_FALSE((bool)It);
+	}
+
+	// -----------------------------------------------------------------------
+	SECTION("EmptyViewFromLiteral")
+	{
+		FUtf8Iterator It(UTF8TEXTVIEW(""));
+		CHECK_FALSE((bool)It);
+	}
+
+	// -----------------------------------------------------------------------
+	SECTION("ASCIIString")
+	{
+		// Every byte in 0x00–0x7F is a single-codepoint sequence
+		TArray<UTF32CHAR> Got = Decode("Hello");
+		TArray<UTF32CHAR> Exp = { 'H', 'e', 'l', 'l', 'o' };
+		CHECK(Got == Exp);
+	}
+
+	// -----------------------------------------------------------------------
+	SECTION("NullByteIsNotSpecial")
+	{
+		// Zero bytes must not terminate iteration; the view length governs the range
+		const UTF8CHAR Bytes[] = { 'A', 0x00, 'B' };
+		TArray<UTF32CHAR> Got = Decode(Bytes, 3);
+		TArray<UTF32CHAR> Exp = { (UTF32CHAR)'A', (UTF32CHAR)0x00u, (UTF32CHAR)'B' };
+		CHECK(Got == Exp);
+	}
+
+	// -----------------------------------------------------------------------
+	SECTION("TwoByteSequence")
+	{
+		// U+00C9 'É'  →  0xC3 0x89
+		const UTF8CHAR Bytes[] = { 0xC3, 0x89 };
+		TArray<UTF32CHAR> Got = Decode(Bytes, 2);
+		TArray<UTF32CHAR> Exp = { (UTF32CHAR)0x00C9u };
+		CHECK(Got == Exp);
+	}
+
+	// -----------------------------------------------------------------------
+	SECTION("TwoByteSequenceMinimum")
+	{
+		// U+0080 — the smallest codepoint requiring a two-byte encoding
+		// Encoded as: 0xC2 0x80
+		const UTF8CHAR Bytes[] = { 0xC2, 0x80 };
+		TArray<UTF32CHAR> Got = Decode(Bytes, 2);
+		TArray<UTF32CHAR> Exp = { (UTF32CHAR)0x0080u };
+		CHECK(Got == Exp);
+	}
+
+	// -----------------------------------------------------------------------
+	SECTION("ThreeByteSequence")
+	{
+		// U+4E2D '中'  →  0xE4 0xB8 0xAD
+		const UTF8CHAR Bytes[] = { 0xE4, 0xB8, 0xAD };
+		TArray<UTF32CHAR> Got = Decode(Bytes, 3);
+		TArray<UTF32CHAR> Exp = { (UTF32CHAR)0x4E2Du };
+		CHECK(Got == Exp);
+	}
+
+	// -----------------------------------------------------------------------
+	SECTION("FourByteSequence")
+	{
+		// U+1F600 '😀'  →  0xF0 0x9F 0x98 0x80
+		const UTF8CHAR Bytes[] = { 0xF0, 0x9F, 0x98, 0x80 };
+		TArray<UTF32CHAR> Got = Decode(Bytes, 4);
+		TArray<UTF32CHAR> Exp = { (UTF32CHAR)0x1F600u };
+		CHECK(Got == Exp);
+	}
+
+	// -----------------------------------------------------------------------
+	SECTION("FourByteSequenceMaximum")
+	{
+		// U+10FFFF — the maximum valid Unicode codepoint
+		// Encoded as: 0xF4 0x8F 0xBF 0xBF
+		const UTF8CHAR Bytes[] = { 0xF4, 0x8F, 0xBF, 0xBF };
+		TArray<UTF32CHAR> Got = Decode(Bytes, 4);
+		TArray<UTF32CHAR> Exp = { (UTF32CHAR)0x10FFFFu };
+		CHECK(Got == Exp);
+	}
+
+	// -----------------------------------------------------------------------
+	SECTION("MixedMultibyteString")
+	{
+		// "A中B" → 0x41, 0xE4 0xB8 0xAD, 0x42
+		const UTF8CHAR Bytes[] = { 0x41, 0xE4, 0xB8, 0xAD, 0x42 };
+		TArray<UTF32CHAR> Got = Decode(Bytes, 5);
+		TArray<UTF32CHAR> Exp = { (UTF32CHAR)'A', (UTF32CHAR)0x4E2Du, (UTF32CHAR)'B' };
+		CHECK(Got == Exp);
+	}
+
+	// -----------------------------------------------------------------------
+	SECTION("InvalidContinuationByteInLeadPosition")
+	{
+		// A bare continuation byte (0x80–0xBF) in lead position is invalid.
+		// Each such byte should produce one replacement character.
+		const UTF8CHAR Bytes[] = { 0x80, 0xBF };
+		TArray<UTF32CHAR> Got = Decode(Bytes, 2);
+		TArray<UTF32CHAR> Exp = { FUtf8Iterator::ReplacementCharacter, FUtf8Iterator::ReplacementCharacter };
+		CHECK(Got == Exp);
+	}
+
+	// -----------------------------------------------------------------------
+	SECTION("OverlongTwoByteSequence")
+	{
+		// 0xC0 0x80 — 0xC0 can never start a valid sequence, so it is consumed
+		// alone.  0x80 is then a bare continuation byte and produces its own
+		// replacement character.
+		const UTF8CHAR Bytes[] = { 0xC0, 0x80 };
+		TArray<UTF32CHAR> Got = Decode(Bytes, 2);
+		TArray<UTF32CHAR> Exp = { FUtf8Iterator::ReplacementCharacter, FUtf8Iterator::ReplacementCharacter };
+		CHECK(Got == Exp);
+	}
+
+	// -----------------------------------------------------------------------
+	SECTION("OverlongThreeByteSequence")
+	{
+		// 0xE0 requires a second byte of 0xA0–0xBF; 0x80 is out of that range, so
+		// 0xE0 is consumed alone.  Each of the two 0x80 continuation bytes then
+		// produces its own replacement character.
+		const UTF8CHAR Bytes[] = { 0xE0, 0x80, 0x80 };
+		TArray<UTF32CHAR> Got = Decode(Bytes, 3);
+		TArray<UTF32CHAR> Exp = { FUtf8Iterator::ReplacementCharacter, FUtf8Iterator::ReplacementCharacter, FUtf8Iterator::ReplacementCharacter };
+		CHECK(Got == Exp);
+	}
+
+	// -----------------------------------------------------------------------
+	SECTION("OverlongFourByteSequence")
+	{
+		// 0xF0 requires a second byte of 0x90–0xBF; 0x80 is out of that range, so
+		// 0xF0 is consumed alone.  Each of the three 0x80 continuation bytes then
+		// produces its own replacement character.
+		const UTF8CHAR Bytes[] = { 0xF0, 0x80, 0x80, 0x80 };
+		TArray<UTF32CHAR> Got = Decode(Bytes, 4);
+		TArray<UTF32CHAR> Exp = { FUtf8Iterator::ReplacementCharacter, FUtf8Iterator::ReplacementCharacter, FUtf8Iterator::ReplacementCharacter, FUtf8Iterator::ReplacementCharacter };
+		CHECK(Got == Exp);
+	}
+
+	// -----------------------------------------------------------------------
+	SECTION("SurrogateCodepoint")
+	{
+		// 0xED requires a second byte of 0x80–0x9F; 0xA0 is out of that range, so
+		// 0xED is consumed alone.  Each of the remaining continuation bytes then
+		// produces its own replacement character.
+		// U+D800 (high surrogate): 0xED 0xA0 0x80
+		const UTF8CHAR HighSurrogate[] = { 0xED, 0xA0, 0x80 };
+		TArray<UTF32CHAR> Got = Decode(HighSurrogate, 3);
+		TArray<UTF32CHAR> Exp = { FUtf8Iterator::ReplacementCharacter, FUtf8Iterator::ReplacementCharacter, FUtf8Iterator::ReplacementCharacter };
+		CHECK(Got == Exp);
+
+		// U+DFFF (low surrogate): 0xED 0xBF 0xBF
+		const UTF8CHAR LowSurrogate[] = { 0xED, 0xBF, 0xBF };
+		Got = Decode(LowSurrogate, 3);
+		CHECK(Got == Exp);
+	}
+
+	// -----------------------------------------------------------------------
+	SECTION("OutOfRangeFourByteSequence")
+	{
+		// 0xF4 requires a second byte of 0x80–0x8F; 0x90 is out of that range, so
+		// 0xF4 is consumed alone.  Each of the three remaining continuation bytes
+		// then produces its own replacement character.
+		// (Would encode U+110000 — one above the maximum valid codepoint.)
+		const UTF8CHAR Bytes[] = { 0xF4, 0x90, 0x80, 0x80 };
+		TArray<UTF32CHAR> Got = Decode(Bytes, 4);
+		TArray<UTF32CHAR> Exp = { FUtf8Iterator::ReplacementCharacter, FUtf8Iterator::ReplacementCharacter, FUtf8Iterator::ReplacementCharacter, FUtf8Iterator::ReplacementCharacter };
+		CHECK(Got == Exp);
+	}
+
+	// -----------------------------------------------------------------------
+	SECTION("IllegalFiveByteSequence")
+	{
+		// 0xF8 0x80 0x80 0x80 0x80 — five-byte lead, illegal per RFC 3629.
+		// The lead byte is consumed alone; each of the four continuation bytes then
+		// produces its own replacement character.
+		const UTF8CHAR Bytes[] = { 0xF8, 0x80, 0x80, 0x80, 0x80 };
+		TArray<UTF32CHAR> Got = Decode(Bytes, 5);
+		TArray<UTF32CHAR> Exp = { FUtf8Iterator::ReplacementCharacter, FUtf8Iterator::ReplacementCharacter, FUtf8Iterator::ReplacementCharacter, FUtf8Iterator::ReplacementCharacter, FUtf8Iterator::ReplacementCharacter };
+		CHECK(Got == Exp);
+	}
+
+	// -----------------------------------------------------------------------
+	SECTION("IllegalSixByteSequence")
+	{
+		// 0xFC 0x80 0x80 0x80 0x80 0x80 — six-byte lead, illegal per RFC 3629.
+		// The lead byte is consumed alone; each of the five continuation bytes then
+		// produces its own replacement character.
+		const UTF8CHAR Bytes[] = { 0xFC, 0x80, 0x80, 0x80, 0x80, 0x80 };
+		TArray<UTF32CHAR> Got = Decode(Bytes, 6);
+		TArray<UTF32CHAR> Exp = { FUtf8Iterator::ReplacementCharacter, FUtf8Iterator::ReplacementCharacter, FUtf8Iterator::ReplacementCharacter, FUtf8Iterator::ReplacementCharacter, FUtf8Iterator::ReplacementCharacter, FUtf8Iterator::ReplacementCharacter };
+		CHECK(Got == Exp);
+	}
+
+	// -----------------------------------------------------------------------
+	SECTION("TruncatedTwoByteAtEnd")
+	{
+		// 0xC3 with no continuation — truncated, produces one replacement and ends.
+		const UTF8CHAR Bytes[] = { 0xC3 };
+		TArray<UTF32CHAR> Got = Decode(Bytes, 1);
+		TArray<UTF32CHAR> Exp = { FUtf8Iterator::ReplacementCharacter };
+		CHECK(Got == Exp);
+	}
+
+	// -----------------------------------------------------------------------
+	SECTION("TruncatedThreeByteOneContinuationAtEnd")
+	{
+		// 0xE4 0xB8 with no third byte — truncated after one continuation.
+		const UTF8CHAR Bytes[] = { 0xE4, 0xB8 };
+		TArray<UTF32CHAR> Got = Decode(Bytes, 2);
+		TArray<UTF32CHAR> Exp = { FUtf8Iterator::ReplacementCharacter };
+		CHECK(Got == Exp);
+	}
+
+	// -----------------------------------------------------------------------
+	SECTION("TruncatedThreeByteNoContinuationAtEnd")
+	{
+		// 0xE4 alone — truncated with no continuations at all.
+		const UTF8CHAR Bytes[] = { 0xE4 };
+		TArray<UTF32CHAR> Got = Decode(Bytes, 1);
+		TArray<UTF32CHAR> Exp = { FUtf8Iterator::ReplacementCharacter };
+		CHECK(Got == Exp);
+	}
+
+	// -----------------------------------------------------------------------
+	SECTION("TruncatedFourByteAtEnd")
+	{
+		// 0xF0 0x9F 0x98 with the final continuation missing.
+		const UTF8CHAR Bytes[] = { 0xF0, 0x9F, 0x98 };
+		TArray<UTF32CHAR> Got = Decode(Bytes, 3);
+		TArray<UTF32CHAR> Exp = { FUtf8Iterator::ReplacementCharacter };
+		CHECK(Got == Exp);
+	}
+
+	// -----------------------------------------------------------------------
+	SECTION("BrokenSequenceNonContinuationNotConsumed")
+	{
+		// 0xC2 followed by 'A' (0x41) — 'A' is not a continuation byte.
+		// The iterator must NOT consume 'A' as part of the invalid sequence;
+		// it should be decoded normally on the next step.
+		const UTF8CHAR Bytes[] = { 0xC2, 0x41 };
+		TArray<UTF32CHAR> Got = Decode(Bytes, 2);
+		TArray<UTF32CHAR> Exp = { FUtf8Iterator::ReplacementCharacter, (UTF32CHAR)'A' };
+		CHECK(Got == Exp);
+	}
+
+	// -----------------------------------------------------------------------
+	SECTION("BrokenThreeByteFirstContinuationInvalid")
+	{
+		// 0xE0 0x41 0x80 — second byte is not a continuation.
+		// Replacement for 0xE0, then 'A' decoded normally, then 0x80 (bare
+		// continuation) yields another replacement.
+		const UTF8CHAR Bytes[] = { 0xE0, 0x41, 0x80 };
+		TArray<UTF32CHAR> Got = Decode(Bytes, 3);
+		TArray<UTF32CHAR> Exp = { FUtf8Iterator::ReplacementCharacter, (UTF32CHAR)'A', FUtf8Iterator::ReplacementCharacter };
+		CHECK(Got == Exp);
+	}
+
+	// -----------------------------------------------------------------------
+	SECTION("BrokenThreeByteSecondContinuationInvalid")
+	{
+		// 0xE4 0xB8 0x41 — third byte is not a continuation.
+		// Replacement for 0xE4 0xB8, then 'A' decoded normally.
+		const UTF8CHAR Bytes[] = { 0xE4, 0xB8, 0x41 };
+		TArray<UTF32CHAR> Got = Decode(Bytes, 3);
+		TArray<UTF32CHAR> Exp = { FUtf8Iterator::ReplacementCharacter, (UTF32CHAR)'A' };
+		CHECK(Got == Exp);
+	}
+
+	// -----------------------------------------------------------------------
+	SECTION("FiveByteLeadFollowedByASCII")
+	{
+		// 0xF8 0x41 — 5-byte lead followed by 'A' (not a continuation).
+		// Only the lead byte should be consumed as the ill-formed sequence;
+		// 'A' must survive to the next iteration.
+		const UTF8CHAR Bytes[] = { 0xF8, 0x41 };
+		TArray<UTF32CHAR> Got = Decode(Bytes, 2);
+		TArray<UTF32CHAR> Exp = { FUtf8Iterator::ReplacementCharacter, (UTF32CHAR)'A' };
+		CHECK(Got == Exp);
+	}
+
+	// -----------------------------------------------------------------------
+	SECTION("MixedValidAndInvalid")
+	{
+		// "A" + invalid 0x80 + U+4E2D '中' + 0xFE (always-invalid byte) + "B"
+		const UTF8CHAR Bytes[] = { 0x41, 0x80, 0xE4, 0xB8, 0xAD, 0xFE, 0x42 };
+		TArray<UTF32CHAR> Got = Decode(Bytes, 7);
+		TArray<UTF32CHAR> Exp = { (UTF32CHAR)'A', FUtf8Iterator::ReplacementCharacter, (UTF32CHAR)0x4E2Du, FUtf8Iterator::ReplacementCharacter, (UTF32CHAR)'B' };
+		CHECK(Got == Exp);
+	}
+
+	// -----------------------------------------------------------------------
+	SECTION("AllInvalidSequences")
+	{
+		// Four bare continuation bytes in a row: each should yield one replacement.
+		const UTF8CHAR Bytes[] = { 0x80, 0x90, 0xA0, 0xBF };
+		TArray<UTF32CHAR> Got = Decode(Bytes, 4);
+		TArray<UTF32CHAR> Exp = { FUtf8Iterator::ReplacementCharacter, FUtf8Iterator::ReplacementCharacter, FUtf8Iterator::ReplacementCharacter, FUtf8Iterator::ReplacementCharacter };
+		CHECK(Got == Exp);
+	}
+
+	// -----------------------------------------------------------------------
+	SECTION("ReplacementCharacterConstant")
+	{
+		CHECK(FUtf8Iterator::ReplacementCharacter == (UTF32CHAR)0xFFFDu);
+	}
+}
+
+#endif // WITH_TESTS
