[InheritedContext]
- Fix lifetime of inherited context scope that should never be destroyed after the base

#Synced-CL 51820749

diff --git a/Engine/Source/Runtime/Core/Public/Async/AsyncWork.h b/Engine/Source/Runtime/Core/Public/Async/AsyncWork.h
--- a/Engine/Source/Runtime/Core/Public/Async/AsyncWork.h
+++ b/Engine/Source/Runtime/Core/Public/Async/AsyncWork.h
@@ -92,10 +92,12 @@ class FAutoDeleteAsyncTask :
 	 **/
 	void DoWork()
 	{
-		UE::FInheritedContextScope InheritedContextScope = RestoreInheritedContext();
-		FScopeCycleCounter Scope(Task.GetStatId(), true);
+		{
+			UE::FInheritedContextScope InheritedContextScope = RestoreInheritedContext();
+			FScopeCycleCounter Scope(Task.GetStatId(), true);
 
-		Task.DoWork();
+			Task.DoWork();
+		}
 		delete this;
 	}
 
@@ -282,10 +284,12 @@ class FAsyncTaskBase
 	**/
 	void DoWork()
 	{
-		UE::FInheritedContextScope InheritedContextScope = RestoreInheritedContext();
-		FScopeCycleCounter Scope(StatId, true);
+		{
+			UE::FInheritedContextScope InheritedContextScope = RestoreInheritedContext();
+			FScopeCycleCounter Scope(StatId, true);
 
-		DoTaskWork();
+			DoTaskWork();
+		}
 		check(WorkNotFinishedCounter.GetValue() == 1);
 		WorkNotFinishedCounter.Decrement();
 	}
diff --git a/Engine/Source/Runtime/RenderCore/Public/RenderingThread.h b/Engine/Source/Runtime/RenderCore/Public/RenderingThread.h
--- a/Engine/Source/Runtime/RenderCore/Public/RenderingThread.h
+++ b/Engine/Source/Runtime/RenderCore/Public/RenderingThread.h
@@ -330,8 +330,10 @@ namespace UE::RenderCommandPipe
 				if (Command->Type == ECommandType::ExecuteFunction)
 				{
 					FExecuteFunctionCommand* FunctionCommand = static_cast<FExecuteFunctionCommand*>(Command);
-					UE::FInheritedContextScope InheritedContextScope = FunctionCommand->RestoreInheritedContext();
-					Lambda(MoveTemp(FunctionCommand->Function), FunctionCommand->Tag);
+					{
+						UE::FInheritedContextScope InheritedContextScope = FunctionCommand->RestoreInheritedContext();
+						Lambda(MoveTemp(FunctionCommand->Function), FunctionCommand->Tag);
+					}
 					FunctionCommand->~FExecuteFunctionCommand();
 				}
 				else
