CompactBinary: Added natvis visualizers

#Synced-CL 50967057

diff --git a/Engine/Source/Runtime/Core/Private/Serialization/CompactBinary.cpp b/Engine/Source/Runtime/Core/Private/Serialization/CompactBinary.cpp
--- a/Engine/Source/Runtime/Core/Private/Serialization/CompactBinary.cpp
+++ b/Engine/Source/Runtime/Core/Private/Serialization/CompactBinary.cpp
@@ -975,3 +975,36 @@ FCompositeBuffer FCbObject::GetBuffer() const
 }
 
 ///////////////////////////////////////////////////////////////////////////////////////////////////////////////////////
+
+namespace UE::CompactBinary::Private
+{
+
+PRAGMA_DISABLE_INTERNAL_WARNINGS;
+FORCENOINLINE void ReferenceDebugTypes(
+	FDebugCb*,
+	FDebugCbObject*,
+	FDebugCbUniformObject*,
+	FDebugCbArray*,
+	FDebugCbUniformArray*,
+	FDebugCbBinary*,
+	FDebugCbString*,
+	FDebugCbIntegerPositive*,
+	FDebugCbIntegerNegative*,
+	FDebugCbFloat32*,
+	FDebugCbFloat64*,
+	FDebugCbObjectAttachment*,
+	FDebugCbBinaryAttachment*,
+	FDebugCbHash*,
+	FDebugCbUuid*,
+	FDebugCbDateTime*,
+	FDebugCbTimeSpan*,
+	FDebugCbObjectId*,
+	FDebugCbCustomById*,
+	FDebugCbCustomByName*)
+{
+}
+PRAGMA_ENABLE_INTERNAL_WARNINGS;
+
+} // UE::CompactBinary::Private
+
+///////////////////////////////////////////////////////////////////////////////////////////////////////////////////////
diff --git a/Engine/Source/Runtime/Core/Public/Serialization/CompactBinary.h b/Engine/Source/Runtime/Core/Public/Serialization/CompactBinary.h
--- a/Engine/Source/Runtime/Core/Public/Serialization/CompactBinary.h
+++ b/Engine/Source/Runtime/Core/Public/Serialization/CompactBinary.h
@@ -1465,6 +1465,64 @@ inline FSharedBuffer FCbField::AsBinary(const FSharedBuffer& Default) &&
 
 ///////////////////////////////////////////////////////////////////////////////////////////////////////////////////////
 
+// Types to support debug visualizers.
+
+// A field/array/object in memory can be visualized in a debugger supporting natvis by casting to this type.
+// (FDebugCb*)Pointer works for any field/array/object that starts with its field type.
+struct UE_INTERNAL FDebugCb {};
+
+// Types that visualize a pointer to a value, *not* a pointer to a field.
+struct UE_INTERNAL FDebugCbObject {};
+struct UE_INTERNAL FDebugCbUniformObject {};
+struct UE_INTERNAL FDebugCbArray {};
+struct UE_INTERNAL FDebugCbUniformArray {};
+struct UE_INTERNAL FDebugCbBinary {};
+struct UE_INTERNAL FDebugCbString {};
+struct UE_INTERNAL FDebugCbIntegerPositive {};
+struct UE_INTERNAL FDebugCbIntegerNegative {};
+struct UE_INTERNAL FDebugCbFloat32 {};
+struct UE_INTERNAL FDebugCbFloat64 {};
+struct UE_INTERNAL FDebugCbObjectAttachment {};
+struct UE_INTERNAL FDebugCbBinaryAttachment {};
+struct UE_INTERNAL FDebugCbHash {};
+struct UE_INTERNAL FDebugCbUuid {};
+struct UE_INTERNAL FDebugCbDateTime {};
+struct UE_INTERNAL FDebugCbTimeSpan {};
+struct UE_INTERNAL FDebugCbObjectId {};
+struct UE_INTERNAL FDebugCbCustomById {};
+struct UE_INTERNAL FDebugCbCustomByName {};
+
+namespace UE::CompactBinary::Private
+{
+
+PRAGMA_DISABLE_INTERNAL_WARNINGS;
+CORE_API void ReferenceDebugTypes(
+	FDebugCb*,
+	FDebugCbObject*,
+	FDebugCbUniformObject*,
+	FDebugCbArray*,
+	FDebugCbUniformArray*,
+	FDebugCbBinary*,
+	FDebugCbString*,
+	FDebugCbIntegerPositive*,
+	FDebugCbIntegerNegative*,
+	FDebugCbFloat32*,
+	FDebugCbFloat64*,
+	FDebugCbObjectAttachment*,
+	FDebugCbBinaryAttachment*,
+	FDebugCbHash*,
+	FDebugCbUuid*,
+	FDebugCbDateTime*,
+	FDebugCbTimeSpan*,
+	FDebugCbObjectId*,
+	FDebugCbCustomById*,
+	FDebugCbCustomByName*);
+PRAGMA_ENABLE_INTERNAL_WARNINGS;
+
+} // UE::CompactBinary::Private
+
+///////////////////////////////////////////////////////////////////////////////////////////////////////////////////////
+
 #if UE_ENABLE_INCLUDE_ORDER_DEPRECATED_IN_5_8
 #include "Memory/CompositeBuffer.h"
 #include "Templates/Function.h"
