Provide Linux with a meaningful FPlatformMemoryStats::GetMemoryPressure() implementation.

- Consider it Critical when used physical (RES) memory falls below 10% of available physical - seems to match the heuristics of the OOM killer that also goes by physical memory.
- Does not change code using the absolute threshold (most other platforms).

#Synced-CL 51825589

diff --git a/Engine/Config/Linux/LinuxEngine.ini b/Engine/Config/Linux/LinuxEngine.ini
--- a/Engine/Config/Linux/LinuxEngine.ini
+++ b/Engine/Config/Linux/LinuxEngine.ini
@@ -19,6 +19,7 @@ DeviceProfileSelectionModule="LinuxDeviceProfileSelector"
 [ConsoleVariables]
 ; larger timeout since drivers may take longer time
 g.TimeoutForBlockOnRenderFence=60000
+memory.MemoryPressureCriticalThresholdPercent=10
 
 [PlatformCrypto]
 PlatformRequiresDataCrypto=True
diff --git a/Engine/Source/Runtime/Core/Private/GenericPlatform/GenericPlatformMemory.cpp b/Engine/Source/Runtime/Core/Private/GenericPlatform/GenericPlatformMemory.cpp
--- a/Engine/Source/Runtime/Core/Private/GenericPlatform/GenericPlatformMemory.cpp
+++ b/Engine/Source/Runtime/Core/Private/GenericPlatform/GenericPlatformMemory.cpp
@@ -62,6 +62,16 @@ namespace GenericPlatformMemory
 			"Where a platform can specifically state it's memory pressure this test maybe ignored.\n"
 			"0 (default) critical pressure will not use the threshold."),
 		ECVF_Default);
+
+	static float GMemoryPressureCriticalThresholdPercent = 0;
+	static FAutoConsoleVariableRef CVarMemoryPressureCriticalThresholdPercent(
+		TEXT("memory.MemoryPressureCriticalThresholdPercent"),
+		GMemoryPressureCriticalThresholdPercent,
+		TEXT("When the available physical memory drops below this percentage of total physical memory, memory stats will consider this to be at critical pressure.\n"
+			"E.g. 10 means Critical when less than 10% of physical RAM is available.\n"
+			"Where a platform can specifically state its memory pressure this test may be ignored.\n"
+			"0 (default) critical pressure will not use the threshold."),
+		ECVF_Default);
 }
 
 #if UE_CHECK_LARGE_ALLOCATIONS
@@ -162,12 +172,33 @@ FGenericPlatformMemoryStats::FGenericPlatformMemoryStats()
 
 FGenericPlatformMemoryStats::EMemoryPressureStatus FGenericPlatformMemoryStats::GetMemoryPressureStatus() const
 {
-	if (GenericPlatformMemory::GMemoryPressureCriticalThresholdMB > 0)
+	const bool bHasMBThreshold      = GenericPlatformMemory::GMemoryPressureCriticalThresholdMB > 0;
+	const bool bHasPercentThreshold = GenericPlatformMemory::GMemoryPressureCriticalThresholdPercent > 0;
+
+	if (!bHasMBThreshold && !bHasPercentThreshold)
 	{
-		float AvailablePhysicalMB = float(AvailablePhysical / 1024 / 1024);
-		return AvailablePhysicalMB < GenericPlatformMemory::GMemoryPressureCriticalThresholdMB ? EMemoryPressureStatus::Critical : EMemoryPressureStatus::Nominal;
+		return EMemoryPressureStatus::Unknown;
 	}
-	return EMemoryPressureStatus::Unknown;
+
+	if (bHasMBThreshold)
+	{
+		const float AvailablePhysicalMB = float(AvailablePhysical / 1024ull / 1024ull);
+		if (AvailablePhysicalMB < GenericPlatformMemory::GMemoryPressureCriticalThresholdMB)
+		{
+			return EMemoryPressureStatus::Critical;
+		}
+	}
+
+	if (bHasPercentThreshold && TotalPhysical > 0)
+	{
+		const float AvailablePhysicalPercent = float(AvailablePhysical) / float(TotalPhysical) * 100.0f;
+		if (AvailablePhysicalPercent < GenericPlatformMemory::GMemoryPressureCriticalThresholdPercent)
+		{
+			return EMemoryPressureStatus::Critical;
+		}
+	}
+
+	return EMemoryPressureStatus::Nominal;
 }
 
 bool FGenericPlatformMemory::bIsOOM = false;
