[InheritedContext]
- Add generic extension mechanism for inherited context

#Synced-CL 51820749

diff --git a/Engine/Source/Runtime/Core/Private/Async/InheritedContext.cpp b/Engine/Source/Runtime/Core/Private/Async/InheritedContext.cpp
--- a/Engine/Source/Runtime/Core/Private/Async/InheritedContext.cpp
+++ b/Engine/Source/Runtime/Core/Private/Async/InheritedContext.cpp
@@ -1,24 +1,375 @@
 // Copyright Epic Games, Inc. All Rights Reserved.
 
 #include "Async/InheritedContext.h"
+#include "HAL/UnrealMemory.h"
+
+#include <atomic>
 
 namespace UE
 {
-	// this method must be defined in a .cpp file to avoid exporting dependency on trace module
-	[[nodiscard]] FInheritedContextScope FInheritedContextBase::RestoreInheritedContext()
-	{
-		return FInheritedContextScope
-		(
-			AppTime
-		#if ENABLE_LOW_LEVEL_MEM_TRACKER
-			, InheritedLLMTags
-		#endif
-		#if (UE_MEMORY_TAGS_TRACE_ENABLED && UE_TRACE_ENABLED)
-			, InheritedMemTag
-		#endif
-		#if (UE_TRACE_METADATA_ENABLED && UE_TRACE_ENABLED) 
-			, InheritedMetadataId
-		#endif
-		);
-	}
-}
\ No newline at end of file
+
+// ---- Thread-local state for extension mechanism ----
+
+// Head of the per-thread linked list of active extension scopes.
+// Pushed/popped by FInheritedContextExtensionScope ctor/dtor.
+static thread_local FInheritedContextExtensionScope* GExtensionScopeHead = nullptr;
+
+// Extension data being propagated through a task chain.
+// Set by FInheritedContextScope ctor when restoring a captured context that has extensions,
+// so that child tasks launched during execution can capture the same extensions.
+static thread_local FInheritedContextExtensionData* GPropagatedExtensionData = nullptr;
+
+// ---- FInheritedContextExtensionData ----
+// Single-allocation storage: [Header][Entry array][Aligned data blob]
+
+struct FInheritedContextExtensionData
+{
+	struct FEntry
+	{
+		FInheritedContextExtension Ext;
+		uint32 Offset; // Byte offset within the data blob for this entry's data
+	};
+
+#if DO_CHECK
+	std::atomic<int32>  RefCount;
+#endif
+	int32  NumEntries;
+	uint32 DataBlockOffset; // Byte offset from `this` to the start of the data blob
+	uint32 TotalAllocSize;
+	uint32 AllocAlignment;
+
+	FEntry*       GetEntries()       { return reinterpret_cast<FEntry*>(this + 1); }
+	const FEntry* GetEntries() const { return reinterpret_cast<const FEntry*>(this + 1); }
+
+	void* GetData(int32 Index)
+	{
+		return reinterpret_cast<uint8*>(this) + DataBlockOffset + GetEntries()[Index].Offset;
+	}
+
+	const void* GetData(int32 Index) const
+	{
+		return reinterpret_cast<const uint8*>(this) + DataBlockOffset + GetEntries()[Index].Offset;
+	}
+
+	// Allocate and populate entries from a list of extension descriptors.
+	// Data values are NOT initialized; caller must invoke Capture on each entry.
+	static FInheritedContextExtensionData* Allocate(const FInheritedContextExtension* const* Extensions, int32 Num)
+	{
+		check(Num > 0);
+
+		// Calculate per-entry data offsets with alignment
+		uint32 DataSize = 0;
+		uint32 MaxDataAlign = 1;
+		TArray<uint32, TInlineAllocator<8>> Offsets;
+		Offsets.SetNumUninitialized(Num);
+
+		for (int32 i = 0; i < Num; ++i)
+		{
+			DataSize = Align(DataSize, Extensions[i]->Impl->DataAlign);
+			Offsets[i] = DataSize;
+			DataSize += Extensions[i]->Impl->DataSize;
+			MaxDataAlign = FMath::Max(MaxDataAlign, Extensions[i]->Impl->DataAlign);
+		}
+
+		// Layout: [FInheritedContextExtensionData][FEntry * Num][Padding][Data blob]
+		const uint32 HeaderAndEntries = static_cast<uint32>(sizeof(FInheritedContextExtensionData) + Num * sizeof(FEntry));
+		const uint32 BlockOffset = Align(HeaderAndEntries, MaxDataAlign);
+		const uint32 TotalSize = BlockOffset + DataSize;
+		const uint32 Alignment = FMath::Max(static_cast<uint32>(alignof(FInheritedContextExtensionData)), MaxDataAlign);
+
+		void* Memory = FMemory::Malloc(TotalSize, Alignment);
+		FInheritedContextExtensionData* Data = static_cast<FInheritedContextExtensionData*>(Memory);
+#if DO_CHECK
+		Data->RefCount.store(0, std::memory_order_relaxed);
+#endif
+		Data->NumEntries      = Num;
+		Data->DataBlockOffset = BlockOffset;
+		Data->TotalAllocSize  = TotalSize;
+		Data->AllocAlignment  = Alignment;
+
+		FEntry* Entries = Data->GetEntries();
+		for (int32 i = 0; i < Num; ++i)
+		{
+			new (&Entries[i]) FEntry{*Extensions[i], Offsets[i]};
+		}
+
+		return Data;
+	}
+
+	// Allocate a block with the same layout as Source (for saved state storage).
+	// Copies header and entry descriptors but NOT data values.
+	static FInheritedContextExtensionData* AllocateLike(const FInheritedContextExtensionData* Source)
+	{
+		void* Memory = FMemory::Malloc(Source->TotalAllocSize, Source->AllocAlignment);
+		FInheritedContextExtensionData* Data = static_cast<FInheritedContextExtensionData*>(Memory);
+
+		// Copy header scalars
+#if DO_CHECK
+		Data->RefCount.store(0, std::memory_order_relaxed);
+#endif
+		Data->NumEntries      = Source->NumEntries;
+		Data->DataBlockOffset = Source->DataBlockOffset;
+		Data->TotalAllocSize  = Source->TotalAllocSize;
+		Data->AllocAlignment  = Source->AllocAlignment;
+
+		// Copy-construct entries (bumps TSharedPtr ref counts)
+		const FEntry* SrcEntries = Source->GetEntries();
+		FEntry* DstEntries = Data->GetEntries();
+		for (int32 i = 0; i < Data->NumEntries; ++i)
+		{
+			new (&DstEntries[i]) FEntry(SrcEntries[i]);
+		}
+
+		return Data;
+	}
+
+	// Call Destroy on each entry's data and free the allocation.
+	void DestroyAndFree()
+	{
+#if DO_CHECK
+		check(RefCount.load(std::memory_order_relaxed) == 0);
+#endif
+		FEntry* Entries = GetEntries();
+		for (int32 i = 0; i < NumEntries; ++i)
+		{
+			if (Entries[i].Ext.Impl->Destroy)
+			{
+				Entries[i].Ext.Impl->Destroy(GetData(i), Entries[i].Ext.Impl->UserData);
+			}
+			Entries[i].~FEntry();
+		}
+		FMemory::Free(this);
+	}
+};
+
+// ---- FInheritedContextExtensionScope ----
+
+FInheritedContextExtensionScope::FInheritedContextExtensionScope(const FInheritedContextExtension& InExtension)
+	: Extension(InExtension)
+	, Next(GExtensionScopeHead)
+{
+	GExtensionScopeHead = this;
+}
+
+FInheritedContextExtensionScope::~FInheritedContextExtensionScope()
+{
+	checkf(GExtensionScopeHead == this, TEXT("FInheritedContextExtensionScope destroyed out of LIFO order"));
+	GExtensionScopeHead = Next;
+}
+
+// ---- FInheritedContextScope ----
+
+FInheritedContextScope::FInheritedContextScope(TSharedPtr<const FAppTime> InAppTime
+	#if ENABLE_LOW_LEVEL_MEM_TRACKER
+		, const FLLMActiveTagsCapture& InInheritedLLMTag
+	#endif
+	#if (UE_MEMORY_TAGS_TRACE_ENABLED && UE_TRACE_ENABLED)
+		, int32 InInheritedMemTag
+	#endif
+	#if (UE_TRACE_METADATA_ENABLED && UE_TRACE_ENABLED)
+		, uint32 InInheritedMetadataId
+	#endif
+		, FInheritedContextExtensionData* InExtensionData
+	)
+	: PrevAppTime(FAppTime::Fork())
+	#if ENABLE_LOW_LEVEL_MEM_TRACKER
+		, LLMScopes(InInheritedLLMTag)
+	#endif
+	#if (UE_MEMORY_TAGS_TRACE_ENABLED && UE_TRACE_ENABLED)
+		, MemScope(InInheritedMemTag)
+	#endif
+	#if (UE_TRACE_METADATA_ENABLED && UE_TRACE_ENABLED)
+		, MetaScope(InInheritedMetadataId)
+	#endif
+{
+	FAppTime::Restore(MoveTemp(InAppTime));
+
+	// Save and update propagated extension data
+	PrevPropagatedData = GPropagatedExtensionData;
+
+	if (InExtensionData && InExtensionData->NumEntries > 0)
+	{
+		// Allocate saved state block with the same layout
+		ExtensionSavedState = FInheritedContextExtensionData::AllocateLike(InExtensionData);
+
+		// Apply each extension: write captured value to thread state, save previous state
+		const FInheritedContextExtensionData::FEntry* Entries = InExtensionData->GetEntries();
+		for (int32 i = 0; i < InExtensionData->NumEntries; ++i)
+		{
+			Entries[i].Ext.Impl->Apply(
+				InExtensionData->GetData(i),
+				ExtensionSavedState->GetData(i),
+				Entries[i].Ext.Impl->UserData
+			);
+		}
+	}
+
+#if DO_CHECK
+	if (InExtensionData)
+	{
+		InExtensionData->RefCount.fetch_add(1, std::memory_order_relaxed);
+	}
+#endif
+
+	// Set propagated data so child tasks launched during this scope capture the same extensions.
+	// Safe: InExtensionData is owned by the FInheritedContextBase which outlives this scope.
+	GPropagatedExtensionData = InExtensionData;
+}
+
+FInheritedContextScope::~FInheritedContextScope()
+{
+	// Restore extensions in reverse order
+	if (ExtensionSavedState)
+	{
+		const FInheritedContextExtensionData::FEntry* Entries = ExtensionSavedState->GetEntries();
+		for (int32 i = ExtensionSavedState->NumEntries - 1; i >= 0; --i)
+		{
+			Entries[i].Ext.Impl->Restore(
+				ExtensionSavedState->GetData(i),
+				Entries[i].Ext.Impl->UserData
+			);
+		}
+
+		ExtensionSavedState->DestroyAndFree();
+	}
+
+#if DO_CHECK
+	if (GPropagatedExtensionData)
+	{
+		check(GPropagatedExtensionData->RefCount.load(std::memory_order_relaxed) > 0);
+		GPropagatedExtensionData->RefCount.fetch_sub(1, std::memory_order_relaxed);
+	}
+#endif
+
+	// Restore previous propagated data
+	GPropagatedExtensionData = PrevPropagatedData;
+
+	FAppTime::Restore(MoveTemp(PrevAppTime));
+}
+
+// ---- FInheritedContextBase ----
+
+void FInheritedContextBase::CaptureInheritedContext()
+{
+	AppTime = FAppTime::Fork();
+
+#if ENABLE_LOW_LEVEL_MEM_TRACKER
+	InheritedLLMTags.CaptureActiveTagData();
+#endif
+
+#if UE_MEMORY_TAGS_TRACE_ENABLED && UE_TRACE_ENABLED
+	InheritedMemTag = MemoryTrace_GetActiveTag();
+#endif
+
+#if UE_TRACE_METADATA_ENABLED && UE_TRACE_ENABLED
+	InheritedMetadataId = UE_TRACE_METADATA_SAVE_STACK();
+#endif
+
+	// Free previous extension data if CaptureInheritedContext was already called.
+	// This handles reuse of the same FInheritedContextBase across multiple dispatches.
+	if (ExtensionData)
+	{
+		ExtensionData->DestroyAndFree();
+		ExtensionData = nullptr;
+	}
+
+	// Capture generic extensions from active scope chain and propagated data
+	if (GExtensionScopeHead || GPropagatedExtensionData)
+	{
+		TArray<const FInheritedContextExtension*, TInlineAllocator<8>> Extensions;
+
+		// Collect from active scope chain (most recent first)
+		for (FInheritedContextExtensionScope* Scope = GExtensionScopeHead; Scope; Scope = Scope->Next)
+		{
+			Extensions.Add(&Scope->Extension);
+		}
+
+		// Collect from propagated data, deduplicating by Impl pointer identity.
+		// The same extension can appear in both the scope chain and propagated data
+		// when a task is retracted (executed inline on the launching thread).
+		// Expected small N (typically 1-3 extensions); linear scan is sufficient.
+		if (GPropagatedExtensionData)
+		{
+			const FInheritedContextExtensionData::FEntry* PropEntries = GPropagatedExtensionData->GetEntries();
+			for (int32 i = 0; i < GPropagatedExtensionData->NumEntries; ++i)
+			{
+				const FInheritedContextExtension* Ext = &PropEntries[i].Ext;
+				bool bDuplicate = false;
+				for (const FInheritedContextExtension* E : Extensions)
+				{
+					if (E->Impl.Get() == Ext->Impl.Get())
+					{
+						bDuplicate = true;
+						break;
+					}
+				}
+				if (!bDuplicate)
+				{
+					Extensions.Add(Ext);
+				}
+			}
+		}
+
+		if (Extensions.Num() > 0)
+		{
+			ExtensionData = FInheritedContextExtensionData::Allocate(Extensions.GetData(), Extensions.Num());
+			for (int32 i = 0; i < Extensions.Num(); ++i)
+			{
+				Extensions[i]->Impl->Capture(ExtensionData->GetData(i), Extensions[i]->Impl->UserData);
+			}
+		}
+	}
+}
+
+FInheritedContextBase::~FInheritedContextBase()
+{
+	if (ExtensionData)
+	{
+		ExtensionData->DestroyAndFree();
+	}
+}
+
+FInheritedContextBase& FInheritedContextBase::operator=(FInheritedContextBase&& Other)
+{
+	if (this != &Other)
+	{
+		if (ExtensionData)
+		{
+			ExtensionData->DestroyAndFree();
+		}
+		AppTime = MoveTemp(Other.AppTime);
+#if ENABLE_LOW_LEVEL_MEM_TRACKER
+		InheritedLLMTags = Other.InheritedLLMTags;
+#endif
+#if UE_MEMORY_TAGS_TRACE_ENABLED && UE_TRACE_ENABLED
+		InheritedMemTag = Other.InheritedMemTag;
+#endif
+#if UE_TRACE_METADATA_ENABLED && UE_TRACE_ENABLED
+		InheritedMetadataId = Other.InheritedMetadataId;
+#endif
+		ExtensionData = Other.ExtensionData;
+		Other.ExtensionData = nullptr;
+	}
+	return *this;
+}
+
+// this method must be defined in a .cpp file to avoid exporting dependency on trace module
+[[nodiscard]] FInheritedContextScope FInheritedContextBase::RestoreInheritedContext()
+{
+	return FInheritedContextScope
+	(
+		AppTime
+	#if ENABLE_LOW_LEVEL_MEM_TRACKER
+		, InheritedLLMTags
+	#endif
+	#if (UE_MEMORY_TAGS_TRACE_ENABLED && UE_TRACE_ENABLED)
+		, InheritedMemTag
+	#endif
+	#if (UE_TRACE_METADATA_ENABLED && UE_TRACE_ENABLED)
+		, InheritedMetadataId
+	#endif
+		, ExtensionData
+	);
+}
+
+}
diff --git a/Engine/Source/Runtime/Core/Public/Async/InheritedContext.h b/Engine/Source/Runtime/Core/Public/Async/InheritedContext.h
--- a/Engine/Source/Runtime/Core/Public/Async/InheritedContext.h
+++ b/Engine/Source/Runtime/Core/Public/Async/InheritedContext.h
@@ -49,45 +49,244 @@ namespace UE
 	};
 #endif
 
-	// Restores an inherited contex for the current scope.
+	/** Opaque per-task storage for generic extension data (defined in .cpp). */
+	struct FInheritedContextExtensionData;
+
+	/**
+	 * Lightweight, copyable handle for a generic inherited context extension.
+	 * Wraps the extension descriptor (callbacks, sizes, UserData) in a TSharedPtr-based pimpl,
+	 * so copying the handle bumps the ref count and the descriptor stays alive automatically.
+	 */
+	struct FInheritedContextExtension
+	{
+		using FCaptureFunc  = void (*)(void* DataDst, void* UserData);
+		using FApplyFunc    = void (*)(const void* DataSrc, void* SaveDst, void* UserData);
+		using FRestoreFunc  = void (*)(const void* SaveSrc, void* UserData);
+		using FDestroyFunc  = void (*)(const void* DataSrc, void* UserData);
+
+		struct FImpl
+		{
+			FCaptureFunc  Capture;    // Read current thread state into DataDst
+			FApplyFunc    Apply;      // Write DataSrc to thread state, save previous to SaveDst
+			FRestoreFunc  Restore;    // Put back previous state from SaveSrc
+			FDestroyFunc  Destroy;    // Destruct data in buffer (nullable for trivial types)
+			uint32        DataSize;
+			uint32        DataAlign;
+			void*         UserData;   // Typically pointer to the variable being propagated
+		};
+
+		TSharedPtr<FImpl> Impl;
+
+		bool IsValid() const { return Impl.IsValid(); }
+
+		template<typename T>
+		T& GetState() const { return *static_cast<T*>(Impl->UserData); }
+	};
+
+	/**
+	 * Creates an FInheritedContextExtension for a variable accessed through a callable.
+	 * The callable is invoked on the current thread each time, so this works correctly with thread_local variables.
+	 * Captures by copy-constructing, applies by assignment. Equivalent to TGuardValue propagation.
+	 * The callable must return an lvalue reference and must be convertible to a function pointer (i.e. non-capturing lambda or free function).
+	 *
+	 * Usage:
+	 *   static FInheritedContextExtension GMyFlagExt = MakeInheritedContextExtension([]() -> bool& { return GMyFlag; });
+	 */
+	template<typename Func>
+	FInheritedContextExtension MakeInheritedContextExtension(Func GetVar)
+	{
+		using ReturnType = decltype(GetVar());
+		static_assert(std::is_lvalue_reference_v<ReturnType>, "Accessor must return an lvalue reference (e.g. []() -> bool& { return GMyFlag; })");
+		using T = std::remove_reference_t<ReturnType>;
+		using FuncPtr = T& (*)();
+
+		static_assert(sizeof(FuncPtr) <= sizeof(void*), "Function pointer must fit in void*");
+
+		FuncPtr FnPtr = GetVar;
+
+		FInheritedContextExtension Ext;
+		Ext.Impl = MakeShared<FInheritedContextExtension::FImpl>();
+		Ext.Impl->Capture = [](void* DataDst, void* UserData)
+		{
+			auto Fn = reinterpret_cast<T& (*)()>(UserData);
+			new (DataDst) T(Fn());
+		};
+		Ext.Impl->Apply = [](const void* DataSrc, void* SaveDst, void* UserData)
+		{
+			auto Fn = reinterpret_cast<T& (*)()>(UserData);
+			new (SaveDst) T(Fn());
+			Fn() = *static_cast<const T*>(DataSrc);
+		};
+		Ext.Impl->Restore = [](const void* SaveSrc, void* UserData)
+		{
+			auto Fn = reinterpret_cast<T& (*)()>(UserData);
+			Fn() = *static_cast<const T*>(SaveSrc);
+		};
+		Ext.Impl->Destroy = [](const void* DataSrc, void* /*UserData*/)
+		{
+			static_cast<const T*>(DataSrc)->~T();
+		};
+		Ext.Impl->DataSize  = sizeof(T);
+		Ext.Impl->DataAlign = alignof(T);
+		Ext.Impl->UserData  = reinterpret_cast<void*>(FnPtr);
+		return Ext;
+	}
+
+	/**
+	 * Creates an FInheritedContextExtension for a getter/setter API.
+	 * Use when the variable is not directly accessible (e.g. accessed through IsSavingPackage()/SetIsSavingPackage()).
+	 * Getter and Setter are compile-time function pointers baked into the generated callbacks.
+	 *
+	 * Usage:
+	 *   static FInheritedContextExtension GExt = MakeInheritedContextExtension<&UE::IsSavingPackage, &UE::SetIsSavingPackage>();
+	 */
+	template<auto Getter, auto Setter>
+	FInheritedContextExtension MakeInheritedContextExtension()
+	{
+		using T = std::decay_t<decltype(Getter())>;
+
+		FInheritedContextExtension Ext;
+		Ext.Impl = MakeShared<FInheritedContextExtension::FImpl>();
+		Ext.Impl->Capture = [](void* DataDst, void* /*UserData*/)
+		{
+			new (DataDst) T(Getter());
+		};
+		Ext.Impl->Apply = [](const void* DataSrc, void* SaveDst, void* /*UserData*/)
+		{
+			new (SaveDst) T(Getter());
+			Setter(*static_cast<const T*>(DataSrc));
+		};
+		Ext.Impl->Restore = [](const void* SaveSrc, void* /*UserData*/)
+		{
+			Setter(*static_cast<const T*>(SaveSrc));
+		};
+		Ext.Impl->Destroy = [](const void* DataSrc, void* /*UserData*/)
+		{
+			static_cast<const T*>(DataSrc)->~T();
+		};
+		Ext.Impl->DataSize  = sizeof(T);
+		Ext.Impl->DataAlign = alignof(T);
+		Ext.Impl->UserData  = nullptr; // Getter/Setter are baked into template callbacks; UserData is not needed
+		return Ext;
+	}
+
+	/**
+	 * Derived FImpl that co-allocates state of type T with the descriptor.
+	 * Defined at namespace scope because local classes cannot have template members.
+	 * When the TSharedPtr<FImpl> ref count hits zero, the control block destructs
+	 * TStatefulContextExtensionImpl<T> which properly destroys State via the type-erased
+	 * deleter stored by TSharedPtr. No virtual destructor needed.
+	 */
+	template<typename T>
+	struct TStatefulContextExtensionImpl : FInheritedContextExtension::FImpl
+	{
+		T State;
+
+		template<typename... CtorArgs>
+		explicit TStatefulContextExtensionImpl(CtorArgs&&... InArgs) : State(Forward<CtorArgs>(InArgs)...) {}
+	};
+
+	/**
+	 * Creates an FInheritedContextExtension with custom scope callbacks and co-allocated mutable state.
+	 * The state of type T is constructed in-place inside a derived FImpl and owned by the handle's
+	 * TSharedPtr, so it lives exactly as long as the last copy of the extension handle.
+	 * ApplyFn() is called when a task starts executing and returns saved state (preserved per-scope).
+	 * RestoreFn(SavedState, SharedState) is called when the task scope ends.
+	 * Access the state after construction via Ext.GetState<T>().
+	 *
+	 * Usage (timing accumulation):
+	 *   double OnStartTiming() { return FPlatformTime::Seconds(); }
+	 *   void OnStopTiming(const double& Start, std::atomic<double>& Acc) {
+	 *       Acc.fetch_add(FPlatformTime::Seconds() - Start, std::memory_order_relaxed);
+	 *   }
+	 *   auto Ext = MakeStatefulInheritedContextExtension<&OnStartTiming, &OnStopTiming, std::atomic<double>>(0.0);
+	 *   FInheritedContextExtensionScope Scope(Ext);
+	 *   // ... launch tasks ...
+	 *   double Result = Ext.GetState<std::atomic<double>>().load();
+	 */
+	template<auto ApplyFn, auto RestoreFn, typename T, typename... ArgsType>
+	FInheritedContextExtension MakeStatefulInheritedContextExtension(ArgsType&&... Args)
+	{
+		using SavedState = std::decay_t<decltype(ApplyFn())>;
+
+		// The capture block stores a T* and the saved state block stores SavedState.
+		// They are separate allocations but share the same DataSize/DataAlign layout,
+		// so size must accommodate the larger of the two.
+		constexpr uint32 Size  = static_cast<uint32>(sizeof(T*) > sizeof(SavedState) ? sizeof(T*) : sizeof(SavedState));
+		constexpr uint32 Align = static_cast<uint32>(alignof(T*) > alignof(SavedState) ? alignof(T*) : alignof(SavedState));
+
+		using FStatefulImpl = TStatefulContextExtensionImpl<T>;
+		auto StatefulRef = MakeShared<FStatefulImpl>(Forward<ArgsType>(Args)...);
+
+		StatefulRef->Capture = [](void* DataDst, void* UserData)
+		{
+			// Store the shared state pointer so it propagates through the chain
+			*static_cast<T**>(DataDst) = static_cast<T*>(UserData);
+		};
+		StatefulRef->Apply = [](const void* /*DataSrc*/, void* SaveDst, void* /*UserData*/)
+		{
+			// Call ApplyFn to begin scope (e.g. record start time)
+			new (SaveDst) SavedState(ApplyFn());
+		};
+		StatefulRef->Restore = [](const void* SaveSrc, void* UserData)
+		{
+			// Call RestoreFn with saved state and shared state
+			RestoreFn(*static_cast<const SavedState*>(SaveSrc), *static_cast<T*>(UserData));
+			static_cast<const SavedState*>(SaveSrc)->~SavedState();
+		};
+		StatefulRef->Destroy  = nullptr; // Captured data is T* (trivially destructible)
+		StatefulRef->DataSize  = Size;
+		StatefulRef->DataAlign = Align;
+		StatefulRef->UserData  = &StatefulRef->State;
+
+		FInheritedContextExtension Ext;
+		Ext.Impl = TSharedPtr<FInheritedContextExtension::FImpl>(StatefulRef);
+		return Ext;
+	}
+
+	/**
+	 * RAII scope that activates an extension for capture.
+	 * While alive, any CaptureInheritedContext() on the same thread includes this extension.
+	 */
+	class FInheritedContextExtensionScope
+	{
+	public:
+		UE_NONCOPYABLE(FInheritedContextExtensionScope);
+
+		CORE_API explicit FInheritedContextExtensionScope(const FInheritedContextExtension& InExtension);
+		CORE_API ~FInheritedContextExtensionScope();
+
+	private:
+		friend class FInheritedContextBase;
+
+		FInheritedContextExtension Extension;
+		FInheritedContextExtensionScope* Next;
+	};
+
+	// Restores an inherited context for the current scope.
 	// An instance must be obtained by calling `FInheritedContextBase::RestoreInheritedContext()`
 	class FInheritedContextScope
 	{
 	public:
-		~FInheritedContextScope()
-		{
-			FAppTime::Restore(MoveTemp(PrevAppTime));
-		}
+		CORE_API ~FInheritedContextScope();
 
 	private:
 		UE_NONCOPYABLE(FInheritedContextScope);
 
 		friend class FInheritedContextBase; // allow construction only by `FInheritedContextBase`
-		
-		FInheritedContextScope(TSharedPtr<const FAppTime> AppTime
+
+		CORE_API FInheritedContextScope(TSharedPtr<const FAppTime> AppTime
 		#if ENABLE_LOW_LEVEL_MEM_TRACKER
 			, const FLLMActiveTagsCapture& InInheritedLLMTag
 		#endif
 		#if (UE_MEMORY_TAGS_TRACE_ENABLED && UE_TRACE_ENABLED)
 			, int32 InInheritedMemTag
 		#endif
-		#if (UE_TRACE_METADATA_ENABLED && UE_TRACE_ENABLED) 
+		#if (UE_TRACE_METADATA_ENABLED && UE_TRACE_ENABLED)
 			, uint32 InInheritedMetadataId
 		#endif
-			)
-			: PrevAppTime(FAppTime::Fork())
-		#if ENABLE_LOW_LEVEL_MEM_TRACKER
-			, LLMScopes(InInheritedLLMTag)
-		#endif
-		#if (UE_MEMORY_TAGS_TRACE_ENABLED && UE_TRACE_ENABLED)
-			, MemScope(InInheritedMemTag)
-		#endif
-		#if (UE_TRACE_METADATA_ENABLED && UE_TRACE_ENABLED) 
-			, MetaScope(InInheritedMetadataId)
-		#endif
-		{
-			FAppTime::Restore(MoveTemp(AppTime));
-		}
+			, FInheritedContextExtensionData* InExtensionData
+		);
 
 		TSharedPtr<const FAppTime> PrevAppTime;
 
@@ -99,34 +298,46 @@ namespace UE
 		FMemScope MemScope;
 	#endif
 
-	#if (UE_TRACE_METADATA_ENABLED && UE_TRACE_ENABLED) 
+	#if (UE_TRACE_METADATA_ENABLED && UE_TRACE_ENABLED)
 		FMetadataRestoreScope MetaScope;
 	#endif
+
+		FInheritedContextExtensionData* ExtensionSavedState = nullptr;  // Previous thread state saved during Apply, consumed by Restore
+		FInheritedContextExtensionData* PrevPropagatedData = nullptr;   // Saved GPropagatedExtensionData, restored when scope ends
 	};
 
-	// this class extends the inherited context (see private members for what the inherited context is) to cover async execution. 
+	// this class extends the inherited context (see private members for what the inherited context is) to cover async execution.
 	// Is intended to be used as a base class, if the inherited context is compiled out it takes 0 space
 	class FInheritedContextBase
 	{
 	public:
-		// must be called in the inherited context, e.g. on launching an async task
-		void CaptureInheritedContext()
-		{
-			AppTime = FAppTime::Fork();
+		FInheritedContextBase() = default;
+		CORE_API ~FInheritedContextBase();
 
+		FInheritedContextBase(const FInheritedContextBase&) = delete;
+		FInheritedContextBase& operator=(const FInheritedContextBase&) = delete;
+
+		FInheritedContextBase(FInheritedContextBase&& Other)
+			: AppTime(MoveTemp(Other.AppTime))
 #if ENABLE_LOW_LEVEL_MEM_TRACKER
-			InheritedLLMTags.CaptureActiveTagData();
+			, InheritedLLMTags(Other.InheritedLLMTags)
 #endif
-
 #if UE_MEMORY_TAGS_TRACE_ENABLED && UE_TRACE_ENABLED
-			InheritedMemTag = MemoryTrace_GetActiveTag();
+			, InheritedMemTag(Other.InheritedMemTag)
 #endif
-
 #if UE_TRACE_METADATA_ENABLED && UE_TRACE_ENABLED
-			InheritedMetadataId = UE_TRACE_METADATA_SAVE_STACK();
+			, InheritedMetadataId(Other.InheritedMetadataId)
 #endif
+			, ExtensionData(Other.ExtensionData)
+		{
+			Other.ExtensionData = nullptr;
 		}
 
+		CORE_API FInheritedContextBase& operator=(FInheritedContextBase&& Other);
+
+		// must be called in the inherited context, e.g. on launching an async task
+		CORE_API void CaptureInheritedContext();
+
 		// must be called where the inherited context should be restored, e.g. at the start of an async task execution
 		[[nodiscard]] CORE_API FInheritedContextScope RestoreInheritedContext();
 
@@ -144,5 +355,7 @@ namespace UE
 #if UE_TRACE_METADATA_ENABLED && UE_TRACE_ENABLED
 		uint32 InheritedMetadataId;
 #endif
+
+		FInheritedContextExtensionData* ExtensionData = nullptr;
 	};
 }
diff --git a/Engine/Source/Runtime/Core/Public/Async/InheritedGuardValue.h b/Engine/Source/Runtime/Core/Public/Async/InheritedGuardValue.h
new file mode 100644
--- /dev/null
+++ b/Engine/Source/Runtime/Core/Public/Async/InheritedGuardValue.h
@@ -0,0 +1,136 @@
+// Copyright Epic Games, Inc. All Rights Reserved.
+
+#pragma once
+
+#include "Async/InheritedContext.h"
+
+/**
+ * Combines TGuardValue with FInheritedContextExtension + FInheritedContextExtensionScope.
+ * Sets a value for the current scope (like TGuardValue) AND propagates it to async tasks
+ * launched within the scope (via the inherited context extension mechanism).
+ *
+ * Usage with callable accessor (creates extension internally):
+ *   static thread_local bool GMyFlag = false;
+ *   TInheritedGuardValue<bool> Guard([]() -> bool& { return GMyFlag; }, true);
+ *
+ * Usage with pre-built extension (avoids allocation per guard):
+ *   static auto GExt = UE::MakeInheritedContextExtension([]() -> bool& { return GMyFlag; });
+ *   TInheritedGuardValue<bool> Guard(GExt, GMyFlag, true);
+ */
+template<typename RefType, typename AssignedType = RefType>
+struct TInheritedGuardValue : private FNoncopyable
+{
+	/**
+	 * Callable constructor — creates the extension internally.
+	 * GetVar must return RefType& (e.g. []() -> bool& { return GMyFlag; }).
+	 * Allocates a TSharedPtr<FImpl> internally, so prefer the pre-built variant for hot paths.
+	 */
+	template<typename Func>
+	[[nodiscard]] TInheritedGuardValue(Func GetVar, const AssignedType& NewValue)
+		: Extension(UE::MakeInheritedContextExtension(GetVar))
+		, Scope(Extension)
+		, RefValue(GetVar())
+		, OriginalValue(RefValue)
+	{
+		RefValue = NewValue;
+	}
+
+	/**
+	 * Pre-built extension constructor — reuses an existing FInheritedContextExtension.
+	 * No allocation; just copies the handle (bumps ref count).
+	 * The extension must have been created with MakeInheritedContextExtension(Func).
+	 */
+	[[nodiscard]] TInheritedGuardValue(const UE::FInheritedContextExtension& InExtension, RefType& ReferenceValue, const AssignedType& NewValue)
+		: Extension(InExtension)
+		, Scope(Extension)
+		, RefValue(ReferenceValue)
+		, OriginalValue(ReferenceValue)
+	{
+		RefValue = NewValue;
+	}
+
+	~TInheritedGuardValue()
+	{
+		RefValue = OriginalValue;
+	}
+
+	/**
+	 * Provides read-only access to the original value of the data being tracked by this struct
+	 *
+	 * @return	a const reference to the original data value
+	 */
+	UE_FORCEINLINE_HINT const AssignedType& GetOriginalValue() const
+	{
+		return OriginalValue;
+	}
+
+private:
+	UE::FInheritedContextExtension Extension;
+	UE::FInheritedContextExtensionScope Scope;
+	RefType& RefValue;
+	AssignedType OriginalValue;
+};
+
+/**
+ * Combines TGuardValueAccessors with FInheritedContextExtension + FInheritedContextExtensionScope.
+ * Like TInheritedGuardValue but for variables accessed through getter/setter functions rather than
+ * direct references. Getter and Setter are compile-time function pointers baked into the extension
+ * callbacks — no TFunction overhead.
+ *
+ * Usage (creates extension internally):
+ *   TInheritedGuardValueAccessors<&UE::IsSavingPackage, &UE::SetIsSavingPackage> Guard(true);
+ *
+ * Usage with pre-built extension (avoids allocation per guard):
+ *   static auto GExt = UE::MakeInheritedContextExtension<&UE::IsSavingPackage, &UE::SetIsSavingPackage>();
+ *   TInheritedGuardValueAccessors<&UE::IsSavingPackage, &UE::SetIsSavingPackage> Guard(GExt, true);
+ */
+template<auto Getter, auto Setter>
+struct TInheritedGuardValueAccessors : private FNoncopyable
+{
+	using ValueType = std::decay_t<decltype(Getter())>;
+
+	/**
+	 * Creates the extension internally.
+	 * Allocates a TSharedPtr<FImpl> internally, so prefer the pre-built variant for hot paths.
+	 */
+	[[nodiscard]] explicit TInheritedGuardValueAccessors(const ValueType& NewValue)
+		: Extension(UE::MakeInheritedContextExtension<Getter, Setter>())
+		, Scope(Extension)
+		, OriginalValue(Getter())
+	{
+		Setter(NewValue);
+	}
+
+	/**
+	 * Pre-built extension constructor — reuses an existing FInheritedContextExtension.
+	 * No allocation; just copies the handle (bumps ref count).
+	 * The extension must have been created with MakeInheritedContextExtension<Getter, Setter>().
+	 */
+	[[nodiscard]] TInheritedGuardValueAccessors(const UE::FInheritedContextExtension& InExtension, const ValueType& NewValue)
+		: Extension(InExtension)
+		, Scope(Extension)
+		, OriginalValue(Getter())
+	{
+		Setter(NewValue);
+	}
+
+	~TInheritedGuardValueAccessors()
+	{
+		Setter(OriginalValue);
+	}
+
+	/**
+	 * Provides read-only access to the original value of the data being tracked by this struct
+	 *
+	 * @return	a const reference to the original data value
+	 */
+	UE_FORCEINLINE_HINT const ValueType& GetOriginalValue() const
+	{
+		return OriginalValue;
+	}
+
+private:
+	UE::FInheritedContextExtension Extension;
+	UE::FInheritedContextExtensionScope Scope;
+	ValueType OriginalValue;
+};
diff --git a/Engine/Source/Runtime/Core/Tests/Async/InheritedContextTest.cpp b/Engine/Source/Runtime/Core/Tests/Async/InheritedContextTest.cpp
new file mode 100644
--- /dev/null
+++ b/Engine/Source/Runtime/Core/Tests/Async/InheritedContextTest.cpp
@@ -0,0 +1,424 @@
+// Copyright Epic Games, Inc. All Rights Reserved.
+
+#include "CoreTypes.h"
+#include "Async/ParallelFor.h"
+#include "Async/InheritedContext.h"
+#include "Async/InheritedGuardValue.h"
+#include "Misc/PackageAccessTracking.h"
+#include "Misc/PackageAccessTrackingOps.h"
+#include "Tasks/Task.h"
+#include "Tests/TestHarnessAdapter.h"
+
+#include <atomic>
+
+#if WITH_TESTS
+
+namespace InheritedContextTests
+{
+	// Getter/setter pair for accessor-based tests
+	static thread_local bool GAccessorFlag = false;
+	static bool GetAccessorFlag() { return GAccessorFlag; }
+	static void SetAccessorFlag(bool Value) { GAccessorFlag = Value; }
+
+	// Helpers for stateful extension tests
+	struct FTaskTime
+	{
+		std::atomic<uint64> GameThreadTime = 0;
+		std::atomic<uint64> WorkerThreadTime = 0;
+	};
+
+	static uint64 GetTime() { return FPlatformTime::Cycles64(); }
+
+	static void SumTime(const uint64& Start, FTaskTime& Time)
+	{
+		if (IsInGameThread())
+		{
+			Time.GameThreadTime.fetch_add(FPlatformTime::Cycles64() - Start, std::memory_order_relaxed);
+		}
+		else
+		{
+			Time.WorkerThreadTime.fetch_add(FPlatformTime::Cycles64() - Start, std::memory_order_relaxed);
+		}
+	}
+
+	// Test simple flag propagation through inherited context extension
+	TEST_CASE_NAMED(FInheritedContextSimpleFlag, "System::Core::Async::InheritedContext::SimpleFlag", "[ApplicationContextMask][EngineFilter]")
+	{
+		using namespace UE::Tasks;
+
+		static thread_local bool GTestFlag = false;
+
+		UE::FInheritedContextExtension GTestFlagExt = UE::MakeInheritedContextExtension([]() -> bool& { return GTestFlag; });
+
+		TGuardValue<bool> Guard(GTestFlag, true);
+
+		// Activate extension - tasks launched here will capture GTestFlag
+		UE::FInheritedContextExtensionScope ActivateFlag(GTestFlagExt);
+
+		// Using an event prevents retraction since we want those tasks to run on worker threads
+		// to validate propagation.
+		FTaskEvent EventA(TEXT("EventA"));
+		Launch(TEXT("TaskA"), [&EventA]()
+		{
+			CHECK(GTestFlag);
+
+			FTaskEvent EventB(TEXT("EventB"));
+			Launch(TEXT("TaskB"), [&EventB]()
+			{
+				CHECK(GTestFlag);
+
+				EventB.Trigger();
+			});
+
+			EventB.Wait();
+
+			EventA.Trigger();
+		});
+
+		EventA.Wait();
+	}
+
+	// Test TInheritedGuardValue propagation (combines guard + extension scope)
+	TEST_CASE_NAMED(FInheritedContextGuardValue, "System::Core::Async::InheritedContext::InheritedGuardValue", "[ApplicationContextMask][EngineFilter]")
+	{
+		using namespace UE::Tasks;
+
+		static thread_local bool GTestFlag = false;
+
+		// Simpler way to propagate guard value on the whole task chain
+		TInheritedGuardValue<bool> Guard([]() -> bool& { return GTestFlag; }, true);
+
+		// Using an event prevents retraction since we want those tasks to run on worker threads
+		// to validate propagation.
+		FTaskEvent EventA(TEXT("EventA"));
+		Launch(TEXT("TaskA"), [&EventA]()
+		{
+			CHECK(GTestFlag);
+
+			FTaskEvent EventB(TEXT("EventB"));
+			Launch(TEXT("TaskB"), [&EventB]()
+			{
+				CHECK(GTestFlag);
+
+				EventB.Trigger();
+			});
+
+			EventB.Wait();
+
+			EventA.Trigger();
+		});
+
+		EventA.Wait();
+	}
+
+	// Test TInheritedGuardValue with pre-built extension (avoids per-guard allocation)
+	TEST_CASE_NAMED(FInheritedContextPreBuiltGuardValue, "System::Core::Async::InheritedContext::PreBuiltGuardValue", "[ApplicationContextMask][EngineFilter]")
+	{
+		using namespace UE::Tasks;
+
+		static thread_local bool GTestFlag = false;
+		static UE::FInheritedContextExtension GTestFlagExt = UE::MakeInheritedContextExtension([]() -> bool& { return GTestFlag; });
+
+		TInheritedGuardValue<bool> Guard(GTestFlagExt, GTestFlag, true);
+
+		FTaskEvent EventA(TEXT("EventA"));
+		Launch(TEXT("TaskA"), [&EventA]()
+		{
+			CHECK(GTestFlag);
+
+			EventA.Trigger();
+		});
+
+		EventA.Wait();
+	}
+
+	// Test getter/setter variant of MakeInheritedContextExtension
+	TEST_CASE_NAMED(FInheritedContextGetterSetter, "System::Core::Async::InheritedContext::GetterSetter", "[ApplicationContextMask][EngineFilter]")
+	{
+		using namespace UE::Tasks;
+
+		UE::FInheritedContextExtension GAccessorFlagExt = UE::MakeInheritedContextExtension<&GetAccessorFlag, &SetAccessorFlag>();
+
+		SetAccessorFlag(true);
+		UE::FInheritedContextExtensionScope ActivateFlag(GAccessorFlagExt);
+
+		FTaskEvent EventA(TEXT("EventA"));
+		Launch(TEXT("TaskA"), [&EventA]()
+		{
+			CHECK(GetAccessorFlag());
+
+			FTaskEvent EventB(TEXT("EventB"));
+			Launch(TEXT("TaskB"), [&EventB]()
+			{
+				CHECK(GetAccessorFlag());
+				EventB.Trigger();
+			});
+
+			EventB.Wait();
+			EventA.Trigger();
+		});
+
+		EventA.Wait();
+
+		SetAccessorFlag(false);
+	}
+
+	// Test TInheritedGuardValueAccessors
+	TEST_CASE_NAMED(FInheritedContextGuardValueAccessors, "System::Core::Async::InheritedContext::GuardValueAccessors", "[ApplicationContextMask][EngineFilter]")
+	{
+		using namespace UE::Tasks;
+
+		TInheritedGuardValueAccessors<&GetAccessorFlag, &SetAccessorFlag> Guard(true);
+
+		FTaskEvent EventA(TEXT("EventA"));
+		Launch(TEXT("TaskA"), [&EventA]()
+		{
+			CHECK(GetAccessorFlag());
+
+			EventA.Trigger();
+		});
+
+		EventA.Wait();
+	}
+
+	// Test propagation through ParallelFor
+	TEST_CASE_NAMED(FInheritedContextParallelFor, "System::Core::Async::InheritedContext::ParallelFor", "[.][ApplicationContextMask][EngineFilter]")
+	{
+		static thread_local bool GTestFlag = false;
+
+		UE::FInheritedContextExtension GTestFlagExt = UE::MakeInheritedContextExtension([]() -> bool& { return GTestFlag; });
+
+		TGuardValue<bool> Guard(GTestFlag, true);
+		UE::FInheritedContextExtensionScope ActivateFlag(GTestFlagExt);
+
+		std::atomic<int32> SuccessCount{0};
+		constexpr int32 NumIterations = 100;
+
+		ParallelFor(NumIterations, [&SuccessCount](int32)
+		{
+			if (GTestFlag)
+			{
+				SuccessCount.fetch_add(1, std::memory_order_relaxed);
+			}
+		}, EParallelForFlags::Unbalanced);
+
+		CHECK(SuccessCount.load() == NumIterations);
+	}
+
+	// Test stateful propagation (timing accumulation across a task chain)
+	TEST_CASE_NAMED(FInheritedContextStatefulContext, "System::Core::Async::InheritedContext::StatefulContext", "[.][ApplicationContextMask][EngineFilter]")
+	{
+		using namespace UE::Tasks;
+
+		// Stateful propagation allows automatic accumulation of metrics (e.g. time) across a task chain.
+		UE::FInheritedContextExtension TaskTimingExtension = UE::MakeStatefulInheritedContextExtension<&GetTime, &SumTime, FTaskTime>();
+
+		// Activate the extension now
+		UE::FInheritedContextExtensionScope ActivateTiming(TaskTimingExtension);
+
+		// Using an event prevents retraction since we want those tasks to run on worker threads
+		// to validate propagation.
+		FTaskEvent EventA(TEXT("EventA"));
+		Launch(TEXT("TaskA"), [&EventA]()
+		{
+			FPlatformProcess::Sleep(0.1f);
+
+			FTaskEvent EventB(TEXT("EventB"));
+			{
+				Launch(TEXT("TaskB"), [&EventB]()
+				{
+					FPlatformProcess::Sleep(0.1f);
+					EventB.Trigger();
+				});
+			}
+
+			EventB.Wait(); // Additional 0.1s for waiting on the other task
+
+			EventA.Trigger();
+		});
+
+		EventA.Wait();
+
+		FTaskTime& TaskTime = TaskTimingExtension.GetState<FTaskTime>();
+
+		uint64 TotalCycles = TaskTime.GameThreadTime.load() + TaskTime.WorkerThreadTime.load();
+		CHECK(FPlatformTime::ToSeconds64(TotalCycles) >= 0.3);
+	}
+
+	// Helpers for deduplication test — uses a global atomic counter so Apply/Restore
+	// side-effects are visible regardless of which thread the task runs on.
+	struct FDedupTracker {};
+	static std::atomic<int32> GApplyCount{0};
+
+	static int32 OnDedupApply()
+	{
+		return GApplyCount.fetch_add(1, std::memory_order_relaxed);
+	}
+
+	static void OnDedupRestore(const int32& /*Saved*/, FDedupTracker& /*State*/)
+	{
+		GApplyCount.fetch_add(-1, std::memory_order_relaxed);
+	}
+
+	// Verify that duplicate extensions are captured only once when the same extension
+	// appears in both the active scope chain and propagated data. This can happen during
+	// task retraction (inline execution on the launching thread) or when a task explicitly
+	// re-activates a scope for an already-propagated extension.
+	TEST_CASE_NAMED(FInheritedContextDeduplication, "System::Core::Async::InheritedContext::Deduplication", "[ApplicationContextMask][EngineFilter]")
+	{
+		using namespace UE::Tasks;
+
+		GApplyCount.store(0, std::memory_order_relaxed);
+
+		UE::FInheritedContextExtension Ext = UE::MakeStatefulInheritedContextExtension<&OnDedupApply, &OnDedupRestore, FDedupTracker>();
+		UE::FInheritedContextExtensionScope Scope(Ext);
+
+		FTaskEvent EventA(TEXT("EventA"));
+		Launch(TEXT("TaskA"), [&Ext, &EventA]() 
+		{
+			// TaskA's scope applied once
+			CHECK(GApplyCount.load(std::memory_order_relaxed) == 1);
+
+			// Re-activate scope for the same extension within the task.
+			// Ext is now in both the scope chain (from this scope) and propagated data
+			// (from TaskA's restored context), creating the dedup scenario.
+			UE::FInheritedContextExtensionScope InnerScope(Ext);
+
+			FTaskEvent EventB(TEXT("EventB"));
+			Launch(TEXT("TaskB"), [&EventB]() 
+			{
+				// With correct dedup: TaskA(1) + TaskB(1) = 2
+				// Without dedup: TaskA(1) + TaskB(2) = 3
+				CHECK(GApplyCount.load(std::memory_order_relaxed) == 2);
+				EventB.Trigger();
+			});
+
+			EventB.Wait();
+			EventA.Trigger();
+		});
+
+		EventA.Wait();
+
+		// All scopes restored
+		CHECK(GApplyCount.load(std::memory_order_relaxed) == 0);
+	}
+
+	// Verify that flags are NOT propagated when no extension scope is active
+	TEST_CASE_NAMED(FInheritedContextNoPropagationWithoutScope, "System::Core::Async::InheritedContext::NoPropagationWithoutScope", "[ApplicationContextMask][EngineFilter]")
+	{
+		using namespace UE::Tasks;
+
+		static thread_local bool GTestFlag = false;
+
+		// Create extension but do NOT activate a scope
+		UE::MakeInheritedContextExtension([]() -> bool& { return GTestFlag; });
+
+		TGuardValue<bool> Guard(GTestFlag, true);
+
+		// No FInheritedContextExtensionScope - flag should not propagate to worker threads
+		FTaskEvent EventA(TEXT("EventA"));
+		Launch(TEXT("TaskA"), [&EventA]() 
+		{
+			CHECK(!GTestFlag);
+
+			FTaskEvent EventB(TEXT("EventB"));
+			Launch(TEXT("TaskB"), [&EventB]() 
+			{
+				CHECK(!GTestFlag);
+
+				EventB.Trigger();
+			});
+
+			EventB.Wait();
+
+			EventA.Trigger();
+		});
+
+		EventA.Wait();
+	}
+
+	TEST_CASE_NAMED(FInheritedContextMultipleExtensions, "System::Core::Async::InheritedContext::MultipleExtensions", "[ApplicationContextMask][EngineFilter]")
+	{
+		using namespace UE::Tasks;
+
+		static thread_local bool GTestFlags[5] = { false };
+
+		// Simpler way to propagate guard value on the whole task chain
+		TInheritedGuardValue<bool> Guard0([]() -> bool& { return GTestFlags[0]; }, true);
+		TInheritedGuardValue<bool> Guard1([]() -> bool& { return GTestFlags[1]; }, true);
+		// Omit 2 on purpose
+		TInheritedGuardValue<bool> Guard3([]() -> bool& { return GTestFlags[3]; }, true);
+		TInheritedGuardValue<bool> Guard4([]() -> bool& { return GTestFlags[4]; }, true);
+
+		FTaskEvent EventA(TEXT("EventA"));
+		Launch(TEXT("TaskA"), [&EventA]()
+		{
+			CHECK(GTestFlags[0]);
+			CHECK(GTestFlags[1]);
+			CHECK(!GTestFlags[2]);
+			CHECK(GTestFlags[3]);
+			CHECK(GTestFlags[4]);
+
+			FTaskEvent EventB(TEXT("EventB"));
+			Launch(TEXT("TaskB"), [&EventB]()
+			{
+				CHECK(GTestFlags[0]);
+				CHECK(GTestFlags[1]);
+				CHECK(!GTestFlags[2]);
+				CHECK(GTestFlags[3]);
+				CHECK(GTestFlags[4]);
+
+				EventB.Trigger();
+			});
+
+			EventB.Wait();
+
+			EventA.Trigger();
+		});
+
+		EventA.Wait();
+	}
+
+	TEST_CASE_NAMED(FInheritedContextExtensionGoesOutOfScopeBeforeExecution, "System::Core::Async::InheritedContext::ExtensionGoesOutOfScopeBeforeExecution", "[ApplicationContextMask][EngineFilter]")
+	{
+		using namespace UE::Tasks;
+
+		static thread_local bool GTestFlag = false;
+
+		FTaskEvent EventA(TEXT("EventA"));
+		FTaskEvent StartEvent(TEXT("StartEvent"));
+
+		{
+			UE::FInheritedContextExtension GTestFlagExt = UE::MakeInheritedContextExtension([]() -> bool& { return GTestFlag; });
+		
+			TGuardValue<bool> Guard(GTestFlag, true);
+			UE::FInheritedContextExtensionScope ActivateFlag(GTestFlagExt);
+
+			Launch(TEXT("TaskA"), [&EventA]()
+			{
+				CHECK(GTestFlag);
+				// test toggling the flag for the next task while we're at it
+				TGuardValue<bool> Guard(GTestFlag, false);
+
+				FTaskEvent EventB(TEXT("EventB"));
+				Launch(TEXT("TaskB"), [&EventB]()
+				{
+					CHECK(!GTestFlag);
+					TGuardValue<bool> Guard(GTestFlag, true);
+
+					EventB.Trigger();
+				});
+
+				EventB.Wait();
+
+				EventA.Trigger();
+			}, StartEvent);
+		}
+
+		CHECK(!GTestFlag);
+		StartEvent.Trigger();
+		EventA.Wait();
+	}
+}
+
+#endif //WITH_TESTS
