Fixed map lookup functions to use KeyFuncs::KeyInitType.
Deprecated the KeyConstPointerType typedefs which were being used instead.

#Synced-CL 51823540

#Synced-CL 51823540

diff --git a/Engine/Plugins/Media/PixelStreaming2/Source/PixelStreaming2/Internal/ThreadSafeMap.h b/Engine/Plugins/Media/PixelStreaming2/Source/PixelStreaming2/Internal/ThreadSafeMap.h
--- a/Engine/Plugins/Media/PixelStreaming2/Source/PixelStreaming2/Internal/ThreadSafeMap.h
+++ b/Engine/Plugins/Media/PixelStreaming2/Source/PixelStreaming2/Internal/ThreadSafeMap.h
@@ -25,7 +25,11 @@ namespace UE::PixelStreaming2
 
 		typedef TSortableMapBase<KeyType, ValueType, SetAllocator, KeyFuncs> Super;
 		typedef typename Super::KeyInitType									 KeyInitType;
+
+PRAGMA_DISABLE_DEPRECATION_WARNINGS
+		UE_DEPRECATED(5.8, "KeyConstPointerType is deprecated - use KeyInitType instead.")
 		typedef typename Super::KeyConstPointerType							 KeyConstPointerType;
+PRAGMA_ENABLE_DEPRECATION_WARNINGS
 
 		TThreadSafeMap() = default;
 		TThreadSafeMap(TThreadSafeMap&&) = default;
@@ -152,7 +156,7 @@ namespace UE::PixelStreaming2
 		 * @param Key the key to search for
 		 * @return whether or not the key was found
 		 */
-		FORCEINLINE ValueType FindAndRemoveChecked(KeyConstPointerType Key)
+		FORCEINLINE ValueType FindAndRemoveChecked(KeyInitType Key)
 		{
 			FScopeLock Lock(&Mutex);
 
@@ -201,14 +205,14 @@ namespace UE::PixelStreaming2
 			}
 		}
 
-		FORCEINLINE ValueType& operator[](KeyConstPointerType Key)
+		FORCEINLINE ValueType& operator[](KeyInitType Key)
 		{
 			FScopeLock Lock(&Mutex);
 
 			return static_cast<Super*>(this)->FindChecked(Key);
 		}
 
-		FORCEINLINE const ValueType& operator[](KeyConstPointerType Key) const
+		FORCEINLINE const ValueType& operator[](KeyInitType Key) const
 		{
 			FScopeLock Lock(&Mutex);
 
@@ -261,7 +265,7 @@ namespace UE::PixelStreaming2
 			return static_cast<Super*>(this)->FindRef(Key);
 		}
 
-		FORCEINLINE int32 Remove(KeyConstPointerType Key)
+		FORCEINLINE int32 Remove(KeyInitType Key)
 		{
 			FScopeLock Lock(&Mutex);
 
@@ -294,7 +298,7 @@ namespace UE::PixelStreaming2
 			return static_cast<Super*>(this)->GetKeys(OutKeys);
 		}
 
-		[[nodiscard]] FORCEINLINE bool Contains(KeyConstPointerType Key)
+		[[nodiscard]] FORCEINLINE bool Contains(KeyInitType Key)
 		{
 			FScopeLock Lock(&Mutex);
 
diff --git a/Engine/Source/Editor/UnrealEd/Private/EditorDomain/EditorDomainUtils.cpp b/Engine/Source/Editor/UnrealEd/Private/EditorDomain/EditorDomainUtils.cpp
--- a/Engine/Source/Editor/UnrealEd/Private/EditorDomain/EditorDomainUtils.cpp
+++ b/Engine/Source/Editor/UnrealEd/Private/EditorDomain/EditorDomainUtils.cpp
@@ -108,7 +108,7 @@ static void EnumSetFlagsAnd(Enum& Output, Enum Mask, Enum A, Enum B)
 
 template <typename KeyType, typename ValueType, typename SetAllocator, typename KeyFuncs>
 static ValueType MapFindRef(const TMap<KeyType, ValueType, SetAllocator, KeyFuncs>& Map,
-	typename TMap<KeyType, ValueType, SetAllocator, KeyFuncs>::KeyConstPointerType Key, ValueType DefaultValue)
+	typename TMap<KeyType, ValueType, SetAllocator, KeyFuncs>::KeyInitType Key, ValueType DefaultValue)
 {
 	const ValueType* FoundValue = Map.Find(Key);
 	return FoundValue ? *FoundValue : DefaultValue;
diff --git a/Engine/Source/Runtime/Core/Public/Containers/Map.h.inl b/Engine/Source/Runtime/Core/Public/Containers/Map.h.inl
--- a/Engine/Source/Runtime/Core/Public/Containers/Map.h.inl
+++ b/Engine/Source/Runtime/Core/Public/Containers/Map.h.inl
@@ -31,8 +31,9 @@ class AUTORTFM_INFER TMAPBASE
 	friend class TMAPBASE;
 
 public:
+	UE_DEPRECATED(5.8, "KeyConstPointerType is deprecated - use KeyInitType instead.")
 	typedef typename TTypeTraits<KeyType  >::ConstPointerType KeyConstPointerType;
-	typedef typename TTypeTraits<KeyType  >::ConstInitType    KeyInitType;
+	typedef typename KeyFuncs::KeyInitType                    KeyInitType;
 	typedef typename TTypeTraits<ValueType>::ConstInitType    ValueInitType;
 	typedef TPair<KeyType, ValueType> ElementType;
 
@@ -405,13 +406,13 @@ public:
 	 * @param InKey The key to remove associated values for.
 	 * @return The number of values that were associated with the key.
 	 */
-	inline int32 Remove(KeyConstPointerType InKey)
+	inline int32 Remove(KeyInitType InKey)
 	{
 		const int32 NumRemovedPairs = Pairs.Remove(InKey);
 		return NumRemovedPairs;
 	}
 
-	inline int32 RemoveStable(KeyConstPointerType InKey)
+	inline int32 RemoveStable(KeyInitType InKey)
 	{
 		const int32 NumRemovedPairs = Pairs.RemoveStable(InKey);
 		return NumRemovedPairs;
@@ -485,7 +486,7 @@ public:
 	 * @return A pointer to the value associated with the specified key, or nullptr if the key isn't contained in this map.  The pointer
 	 *			is only valid until the next change to any key in the map.
 	 */
-	[[nodiscard]] inline ValueType* Find(KeyConstPointerType Key)
+	[[nodiscard]] inline ValueType* Find(KeyInitType Key)
 	{
 		if (auto* Pair = Pairs.Find(Key))
 		{
@@ -494,7 +495,7 @@ public:
 
 		return nullptr;
 	}
-	[[nodiscard]] UE_FORCEINLINE_HINT const ValueType* Find(KeyConstPointerType Key) const
+	[[nodiscard]] UE_FORCEINLINE_HINT const ValueType* Find(KeyInitType Key) const
 	{
 		return const_cast<TMAPBASE*>(this)->Find(Key);
 	}
@@ -631,7 +632,7 @@ public:
 	 * @param Key The key to search for.
 	 * @return The value associated with the specified key, or triggers an assertion if the key does not exist.
 	 */
-	[[nodiscard]] inline const ValueType& FindChecked(KeyConstPointerType Key) const
+	[[nodiscard]] inline const ValueType& FindChecked(KeyInitType Key) const
 	{
 		const auto* Pair = Pairs.Find(Key);
 		check(Pair != nullptr);
@@ -644,7 +645,7 @@ public:
 	 * @param Key The key to search for.
 	 * @return The value associated with the specified key, or triggers an assertion if the key does not exist.
 	 */
-	[[nodiscard]] inline ValueType& FindChecked(KeyConstPointerType Key)
+	[[nodiscard]] inline ValueType& FindChecked(KeyInitType Key)
 	{
 		auto* Pair = Pairs.Find(Key);
 		check(Pair != nullptr);
@@ -657,7 +658,7 @@ public:
 	 * @param Key The key to search for.
 	 * @return The value associated with the specified key, or the default value for the ValueType if the key isn't contained in this map.
 	 */
-	[[nodiscard]] inline ValueType FindRef(KeyConstPointerType Key) const
+	[[nodiscard]] inline ValueType FindRef(KeyInitType Key) const
 	{
 		if (const auto* Pair = Pairs.Find(Key))
 		{
@@ -674,7 +675,7 @@ public:
 	 * @param DefaultValue The fallback value if the key is not found.
 	 * @return The value associated with the specified key, or DefaultValue if the key isn't contained in this map.
 	 */
-	[[nodiscard]] inline ValueType FindRef(KeyConstPointerType Key, ValueType DefaultValue) const
+	[[nodiscard]] inline ValueType FindRef(KeyInitType Key, ValueType DefaultValue) const
 	{
 		if (const auto* Pair = Pairs.Find(Key))
 		{
@@ -706,7 +707,7 @@ public:
 	 * @param Key The key to check for.
 	 * @return true if the map contains the key.
 	 */
-	[[nodiscard]] UE_FORCEINLINE_HINT bool Contains(KeyConstPointerType Key) const
+	[[nodiscard]] UE_FORCEINLINE_HINT bool Contains(KeyInitType Key) const
 	{
 		return Pairs.Contains(Key);
 	}
@@ -1190,7 +1191,6 @@ public:
 
 	typedef TSORTABLEMAPBASE<KeyType, ValueType, SetAllocator, KeyFuncs> Super;
 	typedef typename Super::KeyInitType KeyInitType;
-	typedef typename Super::KeyConstPointerType KeyConstPointerType;
 
 	[[nodiscard]] constexpr TMAP() = default;
 	[[nodiscard]] explicit consteval TMAP(EConstEval)
@@ -1329,7 +1329,7 @@ public:
 	 * @param Key the key to search for
 	 * @return whether or not the key was found
 	 */
-	inline ValueType FindAndRemoveChecked(KeyConstPointerType Key)
+	inline ValueType FindAndRemoveChecked(KeyInitType Key)
 	{
 		const FSetElementId PairId = Super::Pairs.FindId(Key);
 		check(PairId.IsValidId());
@@ -1372,11 +1372,11 @@ public:
 		}
 	}
 
-	[[nodiscard]] UE_FORCEINLINE_HINT ValueType& operator[](KeyConstPointerType Key)
+	[[nodiscard]] UE_FORCEINLINE_HINT ValueType& operator[](KeyInitType Key)
 	{
 		return this->FindChecked(Key);
 	}
-	[[nodiscard]] UE_FORCEINLINE_HINT const ValueType& operator[](KeyConstPointerType Key) const
+	[[nodiscard]] UE_FORCEINLINE_HINT const ValueType& operator[](KeyInitType Key) const
 	{
 		return this->FindChecked(Key);
 	}
@@ -1415,7 +1415,6 @@ class TMULTIMAP : public TSORTABLEMAPBASE<KeyType, ValueType, SetAllocator, KeyF
 
 public:
 	typedef TSORTABLEMAPBASE<KeyType, ValueType, SetAllocator, KeyFuncs> Super;
-	typedef typename Super::KeyConstPointerType KeyConstPointerType;
 	typedef typename Super::KeyInitType KeyInitType;
 	typedef typename Super::ValueInitType ValueInitType;
 
@@ -1590,7 +1589,7 @@ public:
 	 * @param InKey The key to remove associated values for.
 	 * @return The number of values that were associated with the key.
 	 */
-	UE_FORCEINLINE_HINT int32 Remove(KeyConstPointerType InKey)
+	UE_FORCEINLINE_HINT int32 Remove(KeyInitType InKey)
 	{
 		return Super::Remove(InKey);
 	}
diff --git a/Engine/Source/Runtime/Core/Public/Containers/SortedMap.h b/Engine/Source/Runtime/Core/Public/Containers/SortedMap.h
--- a/Engine/Source/Runtime/Core/Public/Containers/SortedMap.h
+++ b/Engine/Source/Runtime/Core/Public/Containers/SortedMap.h
@@ -24,6 +24,7 @@ class TSortedMap
 public:
 	typedef InKeyType      KeyType;
 	typedef InValueType    ValueType;
+	UE_DEPRECATED(5.8, "KeyConstPointerType is deprecated - use KeyInitType instead.")
 	typedef typename TTypeTraits<KeyType  >::ConstPointerType KeyConstPointerType;
 	typedef typename TTypeTraits<KeyType  >::ConstInitType    KeyInitType;
 	typedef typename TTypeTraits<ValueType>::ConstInitType    ValueInitType;
@@ -244,7 +245,7 @@ public:
 	 * @param InKey The key to remove associated values for.
 	 * @return The number of values that were associated with the key.
 	 */
-	inline int32 Remove(KeyConstPointerType InKey, EAllowShrinking AllowShrinking = UE::Core::Private::AllowShrinkingByDefault<ArrayAllocator>())
+	inline int32 Remove(KeyInitType InKey, EAllowShrinking AllowShrinking = UE::Core::Private::AllowShrinkingByDefault<ArrayAllocator>())
 	{
 		int32 RemoveIndex = FindIndex(InKey);
 
@@ -282,7 +283,7 @@ public:
 	 * @param Key The key to search for.
 	 * @return A pointer to the value associated with the specified key, or nullptr if the key isn't contained in this map.  The pointer (only valid until the next change to any key in the map).
 	 */
-	[[nodiscard]] inline ValueType* Find(KeyConstPointerType Key)
+	[[nodiscard]] inline ValueType* Find(KeyInitType Key)
 	{
 		int32 FoundIndex = FindIndex(Key);
 
@@ -294,7 +295,7 @@ public:
 		return nullptr;
 	}
 
-	[[nodiscard]] UE_FORCEINLINE_HINT const ValueType* Find(KeyConstPointerType Key) const
+	[[nodiscard]] UE_FORCEINLINE_HINT const ValueType* Find(KeyInitType Key) const
 	{
 		return const_cast<TSortedMap*>(this)->Find(Key);
 	}
@@ -315,14 +316,14 @@ public:
 	 * @param Key The key to search for.
 	 * @return The value associated with the specified key, or triggers an assertion if the key does not exist.
 	 */
-	[[nodiscard]] inline ValueType& FindChecked(KeyConstPointerType Key)
+	[[nodiscard]] inline ValueType& FindChecked(KeyInitType Key)
 	{
 		ValueType* Value = Find(Key);
 		check(Value != nullptr);
 		return *Value;
 	}
 
-	[[nodiscard]] inline const ValueType& FindChecked(KeyConstPointerType Key) const
+	[[nodiscard]] inline const ValueType& FindChecked(KeyInitType Key) const
 	{
 		const ValueType* Value = Find(Key);
 		check(Value != nullptr);
@@ -335,7 +336,7 @@ public:
 	 * @param Key The key to search for.
 	 * @return The value associated with the specified key, or the default value for the ValueType if the key isn't contained in this map.
 	 */
-	[[nodiscard]] inline ValueType FindRef(KeyConstPointerType Key) const
+	[[nodiscard]] inline ValueType FindRef(KeyInitType Key) const
 	{
 		if (const ValueType* Value = Find(Key))
 		{
@@ -352,7 +353,7 @@ public:
 	 * @param DefaultValue The fallback value if the key is not found.
 	 * @return The value associated with the specified key, or DefaultValue if the key isn't contained in this map.
 	 */
-	[[nodiscard]] inline ValueType FindRef(KeyConstPointerType Key, ValueType DefaultValue) const
+	[[nodiscard]] inline ValueType FindRef(KeyInitType Key, ValueType DefaultValue) const
 	{
 		if (const ValueType* Value = Find(Key))
 		{
@@ -384,7 +385,7 @@ public:
 	 * @param Key The key to check for.
 	 * @return true if the map contains the key.
 	 */
-	[[nodiscard]] inline bool Contains(KeyConstPointerType Key) const
+	[[nodiscard]] inline bool Contains(KeyInitType Key) const
 	{
 		if (Find(Key))
 		{
@@ -470,7 +471,7 @@ public:
 	 * @param Key The key to search for
 	 * @return Whether or not the key was found
 	 */
-	inline ValueType FindAndRemoveChecked(KeyConstPointerType Key, EAllowShrinking AllowShrinking = UE::Core::Private::AllowShrinkingByDefault<ArrayAllocator>())
+	inline ValueType FindAndRemoveChecked(KeyInitType Key, EAllowShrinking AllowShrinking = UE::Core::Private::AllowShrinkingByDefault<ArrayAllocator>())
 	{
 		int32 FoundIndex = FindIndex(Key);
 		check(FoundIndex != INDEX_NONE);
@@ -512,11 +513,11 @@ public:
 		}
 	}
 
-	[[nodiscard]] UE_FORCEINLINE_HINT ValueType& operator[](KeyConstPointerType Key)
+	[[nodiscard]] UE_FORCEINLINE_HINT ValueType& operator[](KeyInitType Key)
 	{
 		return this->FindChecked(Key);
 	}
-	[[nodiscard]] UE_FORCEINLINE_HINT const ValueType& operator[](KeyConstPointerType Key) const
+	[[nodiscard]] UE_FORCEINLINE_HINT const ValueType& operator[](KeyInitType Key) const
 	{
 		return this->FindChecked(Key);
 	}
@@ -572,7 +573,7 @@ private:
 	}
 
 	/** Find index of key */
-	[[nodiscard]] UE_FORCEINLINE_HINT int32 FindIndex(KeyConstPointerType Key) const
+	[[nodiscard]] UE_FORCEINLINE_HINT int32 FindIndex(KeyInitType Key) const
 	{
 		return Algo::BinarySearchBy(Pairs, Key, FKeyForward(), SortPredicate());
 	}
diff --git a/Engine/Source/Runtime/Core/Public/Containers/StripedMap.h b/Engine/Source/Runtime/Core/Public/Containers/StripedMap.h
--- a/Engine/Source/Runtime/Core/Public/Containers/StripedMap.h
+++ b/Engine/Source/Runtime/Core/Public/Containers/StripedMap.h
@@ -56,9 +56,13 @@ class TStripedMapBase
 {
 	static_assert(BucketCount > 0, "The BucketCount needs to be at least 1");
 public:
-	typedef BaseMapType                               MapType;
+	typedef BaseMapType                       MapType;
+	typedef typename BaseMapType::KeyInitType KeyInitType;
+PRAGMA_DISABLE_DEPRECATION_WARNINGS
+	UE_DEPRECATED(5.8, "KeyConstPointerType is deprecated - use KeyInitType instead.")
 	typedef typename BaseMapType::KeyConstPointerType KeyConstPointerType;
-	typedef KeyFuncs                                  KeyFuncsType;
+PRAGMA_ENABLE_DEPRECATION_WARNINGS
+	typedef KeyFuncs                          KeyFuncsType;
 
 protected:
 	struct FDebuggableMutex
@@ -141,7 +145,7 @@ protected:
 	}
 
 	template <typename FunctionType>
-	decltype(auto) ApplyUnlocked(KeyConstPointerType InKey, FunctionType&& InFunction)
+	decltype(auto) ApplyUnlocked(KeyInitType InKey, FunctionType&& InFunction)
 	{
 		return ApplyUnlockedByHash(KeyFuncsType::GetKeyHash(InKey), Forward<FunctionType>(InFunction));
 	}
@@ -156,7 +160,7 @@ protected:
 	}
 
 	template <typename LockType, typename FunctionType>
-	decltype(auto) Apply(KeyConstPointerType InKey, FunctionType&& InFunction) const
+	decltype(auto) Apply(KeyInitType InKey, FunctionType&& InFunction) const
 	{
 		return ApplyByHash<LockType>(KeyFuncsType::GetKeyHash(InKey), Forward<FunctionType>(InFunction));
 	}
@@ -171,43 +175,43 @@ protected:
 	}
 
 	template <typename LockType, typename FunctionType>
-	decltype(auto) Apply(KeyConstPointerType InKey, FunctionType&& InFunction)
+	decltype(auto) Apply(KeyInitType InKey, FunctionType&& InFunction)
 	{
 		return ApplyByHash<LockType>(KeyFuncsType::GetKeyHash(InKey), Forward<FunctionType>(InFunction));
 	}
 
 	template <typename FunctionType>
-	decltype(auto) Write(KeyConstPointerType InKey, FunctionType&& InFunction)
+	decltype(auto) Write(KeyInitType InKey, FunctionType&& InFunction)
 	{
 		return Apply<ExclusiveLockType>(InKey, Forward<FunctionType>(InFunction));
 	}
 
 	template <typename FunctionType>
-	decltype(auto) WriteByHash(uint32 InHash, KeyConstPointerType InKey, FunctionType&& InFunction)
+	decltype(auto) WriteByHash(uint32 InHash, KeyInitType InKey, FunctionType&& InFunction)
 	{
 		return ApplyByHash<ExclusiveLockType>(InHash, InKey, Forward<FunctionType>(InFunction));
 	}
 
 	template <typename FunctionType>
-	decltype(auto) Read(KeyConstPointerType InKey, FunctionType&& InFunction)
+	decltype(auto) Read(KeyInitType InKey, FunctionType&& InFunction)
 	{
 		return Apply<SharedLockType>(InKey, Forward<FunctionType>(InFunction));
 	}
 
 	template <typename FunctionType>
-	decltype(auto) ReadByHash(uint32 InHash, KeyConstPointerType InKey, FunctionType&& InFunction)
+	decltype(auto) ReadByHash(uint32 InHash, KeyInitType InKey, FunctionType&& InFunction)
 	{
 		return ApplyByHash<SharedLockType>(InHash, InKey, Forward<FunctionType>(InFunction));
 	}
 
 	template <typename FunctionType>
-	decltype(auto) Read(KeyConstPointerType InKey, FunctionType&& InFunction) const
+	decltype(auto) Read(KeyInitType InKey, FunctionType&& InFunction) const
 	{
 		return Apply<SharedLockType>(InKey, Forward<FunctionType>(InFunction));
 	}
 
 	template <typename FunctionType>
-	decltype(auto) ReadByHash(uint32 InHash, KeyConstPointerType InKey, FunctionType&& InFunction) const
+	decltype(auto) ReadByHash(uint32 InHash, KeyInitType InKey, FunctionType&& InFunction) const
 	{
 		return ApplyByHash<SharedLockType>(InHash, InKey, Forward<FunctionType>(InFunction));
 	}
@@ -269,7 +273,7 @@ public:
 	 * @param InKey The key to look for.
 	 * @return a copy of the value for this key.
 	 */
-	ValueType FindRef(KeyConstPointerType InKey) const
+	ValueType FindRef(KeyInitType InKey) const
 	{
 		return Read(InKey, [&InKey](uint32 Hash, const MapType& Map)
 			{
@@ -295,7 +299,7 @@ public:
 	 * @return true if the map contains the key and a value was found.
 	 */
 	template <typename FunctionType>
-	bool FindAndApply(KeyConstPointerType InKey, FunctionType&& InFunction) const
+	bool FindAndApply(KeyInitType InKey, FunctionType&& InFunction) const
 	{
 		return Read(InKey, [&InKey, &InFunction](uint32 Hash, const MapType& Map)
 			{
@@ -319,7 +323,7 @@ public:
 	 * @return true if the map contains the key and a value was found.
 	 */
 	template <typename UpdateFunctionType>
-	bool FindAndApply(KeyConstPointerType InKey, UpdateFunctionType&& InUpdateFunction)
+	bool FindAndApply(KeyInitType InKey, UpdateFunctionType&& InUpdateFunction)
 	{
 		return Write(InKey, [&InKey, &InUpdateFunction](uint32 Hash, MapType& Map)
 			{
@@ -341,7 +345,7 @@ public:
 	 * @param InKey The key to check for.
 	 * @return true if the map contains the key.
 	 */
-	[[nodiscard]] bool Contains(KeyConstPointerType InKey) const
+	[[nodiscard]] bool Contains(KeyInitType InKey) const
 	{
 		return Read(InKey, [&InKey](uint32 Hash, const MapType& Map) { return Map.ContainsByHash(Hash, InKey); });
 	}
@@ -578,13 +582,13 @@ public:
 	 * @param InKey The key to remove associated values for.
 	 * @return The number of values that were associated with the key.
 	 */
-	int32 Remove(KeyConstPointerType InKey)
+	int32 Remove(KeyInitType InKey)
 	{
 		return Write(InKey, [&InKey](uint32 Hash, MapType& Map) { return Map.RemoveByHash(Hash, InKey); });
 	}
 	
 	/** See Remove() and class documentation section on ByHash() functions */
-	int32 RemoveByHash(uint32 InKeyHash, KeyConstPointerType InKey)
+	int32 RemoveByHash(uint32 InKeyHash, KeyInitType InKey)
 	{
 		return WriteByHash(InKeyHash, InKey, [&InKey](uint32 Hash, MapType& Map) { return Map.RemoveByHash(Hash, InKey); });
 	}
@@ -597,7 +601,7 @@ public:
 	 * @return The number of values that were removed.
 	 */
 	template <typename PredicateType>
-	int32 RemoveIf(KeyConstPointerType InKey, PredicateType&& InPredicate)
+	int32 RemoveIf(KeyInitType InKey, PredicateType&& InPredicate)
 	{
 		return Write(InKey,
 			[&InKey, &InPredicate](uint32 Hash, MapType& Map)
@@ -675,7 +679,7 @@ public:
 	 * @param OutRemovedValue If found, the value that was removed (not modified if the key was not found)
 	 * @return whether or not the key was found
 	 */
-	bool RemoveAndCopyValue(KeyConstPointerType InKey, ValueType& OutRemovedValue)
+	bool RemoveAndCopyValue(KeyInitType InKey, ValueType& OutRemovedValue)
 	{
 		return Write(InKey, [&InKey, &OutRemovedValue](uint32 Hash, MapType& Map) 
 			{
@@ -691,7 +695,7 @@ public:
 	 * @param InKey The key to search for
 	 * @return      The value that was removed
 	 */
-	ValueType FindAndRemoveChecked(KeyConstPointerType InKey)
+	ValueType FindAndRemoveChecked(KeyInitType InKey)
 	{
 		return Write(InKey, [&InKey](uint32 Hash, MapType& Map)
 			{
@@ -937,7 +941,7 @@ public:
 	 * @param bMaintainOrder true if the Values array should be in the same order as the map's pairs.
 	 */
 	template<typename Allocator>
-	void MultiFind(typename Super::KeyConstPointerType InKey, TArray<ValueType, Allocator>& OutValues, bool bMaintainOrder = false) const
+	void MultiFind(typename Super::KeyInitType InKey, TArray<ValueType, Allocator>& OutValues, bool bMaintainOrder = false) const
 	{
 		Super::Read(InKey, [&InKey, &OutValues, bMaintainOrder](uint32 Hash, const typename Super::MapType& Map) { Map.MultiFind(InKey, OutValues, bMaintainOrder); });
 	}
diff --git a/Engine/Source/Runtime/Core/Public/Misc/ConfigCacheIni.h b/Engine/Source/Runtime/Core/Public/Misc/ConfigCacheIni.h
--- a/Engine/Source/Runtime/Core/Public/Misc/ConfigCacheIni.h
+++ b/Engine/Source/Runtime/Core/Public/Misc/ConfigCacheIni.h
@@ -673,7 +673,7 @@ public:
 		return FConfigFileMap::GetKeys(Keys);
 	}
 
-	inline int32 Remove(KeyConstPointerType InKey)
+	inline int32 Remove(KeyInitType InKey)
 	{
 		UE::TWriteScopeLock ScopeLock(ConfigFileMapLock);
 		return FConfigFileMap::Remove(InKey);
