Structured logging support for indexing into arrays when subobject paths are supported

#Synced-CL 51675159
#Synced-CL 51675159

diff --git a/Engine/Source/Runtime/Core/Private/Logging/StructuredLog.cpp b/Engine/Source/Runtime/Core/Private/Logging/StructuredLog.cpp
--- a/Engine/Source/Runtime/Core/Private/Logging/StructuredLog.cpp
+++ b/Engine/Source/Runtime/Core/Private/Logging/StructuredLog.cpp
@@ -31,6 +31,8 @@
 #include "Serialization/CompactBinaryValue.h"
 #include "Serialization/CompactBinaryWriter.h"
 #include "Serialization/VarInt.h"
+#include "String/LexFromString.h"
+#include "String/Numeric.h"
 #include "String/Split.h"
 #include "Templates/Function.h"
 #include <cstdarg>
@@ -57,7 +59,7 @@ CORE_API bool GPrependLogContextToLogMessage = false;
 
 static constexpr ANSICHAR GFieldPathDelimiter = '/';
 static constexpr FAsciiSet GValidLogFieldName("0123456789ABCDEFGHIJKLMNOPQRSTUVWXYZabcdefghijklmnopqrstuvwxyz_");
-static constexpr FAsciiSet GValidLogFieldPath = GValidLogFieldName | FAsciiSet({GFieldPathDelimiter, '\0'});
+static constexpr FAsciiSet GValidLogFieldPath = GValidLogFieldName | FAsciiSet({GFieldPathDelimiter, '[', ']', '\0'});
 
 static constexpr FAnsiStringView GLogContextsFieldName = ANSITEXTVIEW("$Contexts");
 
@@ -339,12 +341,42 @@ public:
 	FCbFieldView FindByPath(FUtf8StringView Path, int32 IndexHint = -1)
 	{
 		FUtf8StringView Name = Path;
+		FUtf8StringView IndexDigits;
 		bool bMore = String::SplitFirstChar(Path, UTF8CHAR(GFieldPathDelimiter), Name, Path);
+		bool bArray = Name.EndsWith(']') && String::SplitLastChar(Name, UTF8CHAR('['), Name, IndexDigits);
 		FCbFieldView Field = Find(Name, IndexHint);
-		while (bMore)
+		for (;;)
 		{
+			if (UNLIKELY(bArray))
+			{
+				IndexDigits.RemoveSuffix(1);
+				checkf(String::IsNumericOnlyDigits(IndexDigits), TEXT("Log format has non-numeric array index '%s'. [[%s]]"),
+					*WriteToString<32>(IndexDigits), StringCast<TCHAR>(Format).Get());
+				FCbArrayView Array = Field.AsArrayView();
+				checkf(!Field.HasError(), TEXT("Log format requires field '%s' but the corresponding field is not an array. [[%s]]"),
+					*WriteToString<32>(Name, '[', IndexDigits, ']'), StringCast<TCHAR>(Format).Get());
+				int32 ArrayIndex = 0;
+				LexFromString(ArrayIndex, IndexDigits);
+				checkf(ArrayIndex < Array.Num(), TEXT("Log format requires item at index %d in array '%s' but the array only contains %d items. [[%s]]"),
+					ArrayIndex, *WriteToString<32>(Name), Array.Num(), StringCast<TCHAR>(Format).Get());
+				for (FCbFieldView ArrayIt : Array)
+				{
+					if (ArrayIndex-- == 0)
+					{
+						Field = ArrayIt;
+						break;
+					}
+				}
+			}
+
+			if (!bMore)
+			{
+				break;
+			}
+
 			Name = Path;
 			bMore = String::SplitFirstChar(Path, UTF8CHAR(GFieldPathDelimiter), Name, Path);
+			bArray = Name.EndsWith(']') && String::SplitLastChar(Name, UTF8CHAR('['), Name, IndexDigits);
 			Field = Field.AsObjectView().FindView(Name);
 			checkf(Field, TEXT("Log format requires field '%s' which was not provided. [[%s]]"),
 				*WriteToString<32>(Path), StringCast<TCHAR>(Format).Get());
@@ -571,9 +603,9 @@ FLogTemplate* FLogTemplate::Create(
 		const FAnsiStringView FieldPath = FieldPathData.ToView().RightChop(FieldPathIndex);
 		FieldPathSizes.Add(FieldPath.Len());
 
-		const bool bHasSubObjectReference = !!Algo::Find(FieldPath, GFieldPathDelimiter);
+		const bool bHasSubObjectReference = FAsciiSet::HasAny(FieldPath, GValidLogFieldPath & ~GValidLogFieldName);
 		checkf(!bHasSubObjectReference || Options.bAllowSubObjectReferences,
-			TEXT("Log format has a sub-object reference (%c) in a context that does not allow them. [[%s]]"),
+			TEXT("Log format has a sub-object reference (%c or []) in a context that does not allow them. [[%s]]"),
 			GFieldPathDelimiter, StringCast<TCHAR>(Format).Get());
 
 		if (bFindFields && !bPositional)
@@ -683,9 +715,9 @@ FLogTemplate* FLogTemplate::CreateLocalized(
 		FieldPathData.Append(FieldStart + 1, FieldLen - 2);
 		const FAnsiStringView FieldPath = FieldPathData.ToView().RightChop(FieldPathIndex);
 
-		const bool bHasSubObjectReference = !!Algo::Find(FieldPath, GFieldPathDelimiter);
+		const bool bHasSubObjectReference = FAsciiSet::HasAny(FieldPath, GValidLogFieldPath & ~GValidLogFieldName);
 		checkf(!bHasSubObjectReference || Options.bAllowSubObjectReferences,
-			TEXT("Log format has a sub-object reference (%c) in a context that does not allow them. [[%s]]"),
+			TEXT("Log format has a sub-object reference (%c or []) in a context that does not allow them. [[%s]]"),
 			GFieldPathDelimiter, StringCast<TCHAR>(Format).Get());
 
 		// Check if the field path has been seen and skip it if it has.
diff --git a/Engine/Source/Runtime/Core/Public/Logging/StructuredLogFormat.h b/Engine/Source/Runtime/Core/Public/Logging/StructuredLogFormat.h
--- a/Engine/Source/Runtime/Core/Public/Logging/StructuredLogFormat.h
+++ b/Engine/Source/Runtime/Core/Public/Logging/StructuredLogFormat.h
@@ -79,7 +79,12 @@ namespace UE
 /** Options to control how log templates are parsed and perform formatting. */
 struct FLogTemplateOptions
 {
-	/** If true, allow field references of the form A.B.C to access fields of objects within objects. */
+	/**
+	 * If true, allow field references at a depth greater than 1.
+	 *
+	 * Access objects within objects using a path of the form A/B/C.
+	 * Access items within arrays using a path of the form A[0]/B/C[1].
+	 */
 	bool bAllowSubObjectReferences = false;
 };
 
diff --git a/Engine/Source/Runtime/Core/Tests/Logging/StructuredLogBenchmarks.cpp b/Engine/Source/Runtime/Core/Tests/Logging/StructuredLogBenchmarks.cpp
new file mode 100644
--- /dev/null
+++ b/Engine/Source/Runtime/Core/Tests/Logging/StructuredLogBenchmarks.cpp
@@ -0,0 +1,318 @@
+// Copyright Epic Games, Inc. All Rights Reserved.
+
+#if WITH_LOW_LEVEL_TESTS
+
+#include "Containers/StaticArray.h"
+#include "Containers/UnrealString.h"
+#include "Delegates/IntegerSequence.h"
+#include "Logging/StructuredLogFormat.h"
+#include "Misc/StringBuilder.h"
+#include "Serialization/CompactBinary.h"
+#include "Serialization/CompactBinaryWriter.h"
+#include "Templates/Invoke.h"
+#include "Templates/Overload.h"
+#include "Templates/Tuple.h"
+#include "Tests/TestHarnessAdapter.h"
+
+#include <catch2/benchmark/catch_benchmark.hpp>
+
+namespace UE
+{
+static inline constexpr int32 NumTemplateTypes = 10;
+
+// Use the same values for each test for insertion into compact binary
+constexpr int32 IntegerValues[NumTemplateTypes] = {
+	2118,
+	6695,
+	15413,
+	212350,
+	534561,
+	3940185,
+	4538994,
+	22786616,
+	96034783,
+	226063757,
+};
+
+constexpr FUtf8StringView StringValues[NumTemplateTypes] = {
+	UTF8TEXTVIEW("Param2118"),
+	UTF8TEXTVIEW("Param6695"),
+	UTF8TEXTVIEW("Param15413"),
+	UTF8TEXTVIEW("Param212350"),
+	UTF8TEXTVIEW("Param534561"),
+	UTF8TEXTVIEW("Param3940185"),
+	UTF8TEXTVIEW("Param4538994"),
+	UTF8TEXTVIEW("Param22786616"),
+	UTF8TEXTVIEW("Param96034783"),
+	UTF8TEXTVIEW("Param226063757"),
+};
+
+// Helper to write the same structure into multiple FCbWriter with different leaf field types
+template<int32 NumWriters>
+class TFieldWriters
+{
+	TStaticArray<FCbWriter, NumWriters> Writers;
+
+	static void AddValue(FCbWriter& W, int32 Value)
+	{
+		W.AddInteger(Value);
+	}
+
+	static void AddValue(FCbWriter& W, FUtf8StringView Value)
+	{
+		W.AddString(Value);
+	}
+
+public:
+	TFieldWriters() = default;
+	UE_NONCOPYABLE(TFieldWriters);
+
+	void BeginObject(FUtf8StringView Name)
+	{
+		for (int32 i=0; i < NumWriters; ++i)
+		{
+			Writers[i].BeginObject(Name);
+		}
+	}
+	
+	void EndObject()
+	{
+		for (int32 i=0; i < NumWriters; ++i)
+		{
+			Writers[i].EndObject();
+		}
+	}
+
+	void BeginArray(FUtf8StringView Name)
+	{
+		for (int32 i=0; i < NumWriters; ++i)
+		{
+			Writers[i].BeginArray(Name);
+		}
+	}
+	
+	void EndArray()
+	{
+		for (int32 i=0; i < NumWriters; ++i)
+		{
+			Writers[i].EndArray();
+		}
+	}
+
+	template<typename... ArgTypes>
+	void AddValues(ArgTypes&&... Args)
+	{
+		static_assert(sizeof...(Args) == NumWriters);
+		([&]<uint32... Indices>(TIntegerSequence<uint32, Indices...>)
+		{
+			(AddValue(Writers[Indices], Args), ...);
+		})(TMakeIntegerSequence<uint32, NumWriters>{});
+	}
+
+	template<typename... ArgTypes>
+	void AddValues(FUtf8StringView Name, ArgTypes&&... Args)
+	{
+		static_assert(sizeof...(Args) == NumWriters);
+		for (int32 i=0; i < NumWriters; ++i)
+		{
+			Writers[i].SetName(Name);
+		}
+
+		AddValues(Forward<ArgTypes>(Args)...);
+	}
+
+	// Save each writer and return a tuple of FCbFieldIterator
+	TStaticArray<FCbFieldIterator, NumWriters> Save()
+	{
+		TStaticArray<FCbFieldIterator, NumWriters> Cb;
+		for (int32 i=0; i < NumWriters; ++i)
+		{
+			Cb[i] = Writers[i].Save();
+		}
+		return Cb;
+	}
+};
+
+TEST_CASE_NAMED(FStructuredLogBenchmark_FormatName, "Core::Logging::StructuredLog::Benchmark::FormatName", "[Core][!benchmark]")
+{
+	TArray<FString> TemplateStrings;
+	for (int32 Variant=0; Variant < NumTemplateTypes; ++Variant)
+	{
+		TStringBuilder<1024> Template;
+		Template << TEXTVIEW("Structured log test: ");
+		for (int32 ParamIndex=0; ParamIndex <= Variant; ++ParamIndex)
+		{
+			Template << TEXTVIEW("{Param") << ParamIndex << TEXTVIEW("} ");
+		}
+		TemplateStrings.Emplace(Template.ToString());
+	}
+	TArray<UE::FUniqueLogTemplate> Templates;
+	Templates.Reserve(NumTemplateTypes);
+	for (int32 i=0; i < NumTemplateTypes; ++i)			
+	{
+		Templates.Emplace(*TemplateStrings[i]);
+	}
+
+	TFieldWriters<2> Writers;
+	for (int32 i = 0; i < NumTemplateTypes; ++i)
+	{
+		Writers.AddValues(WriteToUtf8String<64>(TEXTVIEW("Param"), i).ToView(),
+			IntegerValues[i],
+			StringValues[i]);
+	}
+	TStaticArray<FCbFieldIterator, 2> Cb = Writers.Save();
+
+	TStringBuilder<1024> Builder;
+	BENCHMARK("IntegerFields")
+	{
+		for (int32 i=0; i < Templates.Num(); ++i)
+		{
+			TStringBuilder<1024> Builder;
+			Templates[i].FormatTo(Builder, Cb[0]);
+		}
+	};
+
+	BENCHMARK("StringFields")
+	{
+		for (int32 i=0; i < Templates.Num(); ++i)
+		{
+			TStringBuilder<1024> Builder;
+			Templates[i].FormatTo(Builder, Cb[1]);
+		}
+	};
+}
+
+TEST_CASE_NAMED(FStructuredLogBenchmark_FormatPath, "Core::Logging::StructuredLog::Benchmark::FormatPath", "[Core][!benchmark]")
+{
+	TArray<FString> TemplateStrings;
+	for (int32 Variant=0; Variant < NumTemplateTypes; ++Variant)			
+	{
+		TStringBuilder<1024> Template;
+		Template << TEXTVIEW("Structured log test: ");
+		// For variant n, add n fields
+		// Each field is nested n levels under an object called Root, e.g. Root/Param0, Root/Param1/Param0
+		for (int32 j=0; j < Variant; ++j)
+		{
+			Template << TEXTVIEW("{Root"); 
+			for (int32 k=j; k >= 0; --k)
+			{
+				Template << TEXTVIEW("/Param") << k;
+			}
+			Template << TEXTVIEW("} ");
+		}
+		TemplateStrings.Emplace(Template.ToString());
+	}
+	TArray<UE::FUniqueLogTemplate> Templates;
+	Templates.Reserve(NumTemplateTypes);
+	for (int32 i=0; i < NumTemplateTypes; ++i)			
+	{
+		Templates.Emplace(*TemplateStrings[i], UE::FLogTemplateOptions{ .bAllowSubObjectReferences = true });
+	}
+
+	TFieldWriters<2> Writers;
+	Writers.BeginObject(UTF8TEXTVIEW("Root"));
+	for (int32 i=0; i < NumTemplateTypes; ++i)
+	{
+		for (int32 j = i; j > 0; --j)
+		{
+			Writers.BeginObject(WriteToUtf8String<64>(TEXTVIEW("Param"), j).ToView());
+		}
+		Writers.AddValues(UTF8TEXTVIEW("Param0"), IntegerValues[i], StringValues[i]);
+		for (int32 j = i; j > 0; --j)
+		{
+			Writers.EndObject();
+		}
+	}
+	Writers.EndObject();
+	TStaticArray<FCbFieldIterator, 2> Cb = Writers.Save();
+	BENCHMARK("IntegerFields")
+	{
+		for (int32 i=0; i < Templates.Num(); ++i)
+		{
+			TStringBuilder<1024> Builder;
+			Templates[i].FormatTo(Builder, Cb[0]);
+		}
+	};
+
+	BENCHMARK("StringFields")
+	{
+		for (int32 i=0; i < Templates.Num(); ++i)
+		{
+			TStringBuilder<1024> Builder;
+			Templates[i].FormatTo(Builder, Cb[1]);
+		}
+	};
+}
+
+TEST_CASE_NAMED(FStructuredLogBenchmark_FormatPath, "Core::Logging::StructuredLog::Benchmark::FormatPathArrayIndex", "[Core][!benchmark]")
+{
+	TArray<FString> TemplateStrings;
+	for (int32 i=0; i < NumTemplateTypes; ++i)			
+	{
+		int32 Variant = i % NumTemplateTypes;
+		TStringBuilder<1024> Template;
+		Template << TEXTVIEW("Structured log test: ");
+		// For variant n, add n fields
+		// Each field is nested n levels under an object called Root and terminates with an array index n, e.g. Root/Param0[0], Root/Param1/Param0[1]
+		for (int32 j=0; j < Variant; ++j)
+		{
+			Template << TEXTVIEW("{Root"); 
+			for (int32 k=j; k >= 0; --k)
+			{
+				Template << TEXTVIEW("/Param") << k;
+			}
+			Template << TEXTVIEW("[") << j << TEXTVIEW("]} ");
+		}
+		TemplateStrings.Emplace(Template.ToString());
+	}
+	TArray<UE::FUniqueLogTemplate> Templates;
+	Templates.Reserve(NumTemplateTypes);
+	for (int32 i=0; i < NumTemplateTypes; ++i)			
+	{
+		Templates.Emplace(*TemplateStrings[i], UE::FLogTemplateOptions{ .bAllowSubObjectReferences = true });
+	}
+
+	TFieldWriters<2> Writers;
+	Writers.BeginObject(UTF8TEXTVIEW("Root"));
+	for (int32 i=0; i < NumTemplateTypes; ++i)
+	{
+		for (int32 j = i; j > 0; --j)
+		{
+			Writers.BeginObject(WriteToUtf8String<64>(TEXTVIEW("Param"), j).ToView());
+		}
+		Writers.BeginArray(UTF8TEXTVIEW("Param0"));
+		for (int32 j = 0; j < NumTemplateTypes; ++j)
+		{
+			Writers.AddValues(IntegerValues[j], StringValues[j]);
+		}
+		Writers.EndArray();
+		for (int32 j = i; j > 0; --j)
+		{
+			Writers.EndObject();
+		}
+	}
+	Writers.EndObject();
+	TStaticArray<FCbFieldIterator, 2> Cb= Writers.Save();
+
+	BENCHMARK("IntegerFields")
+	{
+		for (int32 i=0; i < Templates.Num(); ++i)
+		{
+			TStringBuilder<1024> Builder;
+			Templates[i].FormatTo(Builder, Cb[0]);
+		}
+	};
+
+	BENCHMARK("StringFields")
+	{
+		for (int32 i=0; i < Templates.Num(); ++i)
+		{
+			TStringBuilder<1024> Builder;
+			Templates[i].FormatTo(Builder, Cb[1]);
+		}
+	};
+}
+
+}
+
+#endif // WITH_LOW_LEVEL_TESTS
\ No newline at end of file
diff --git a/Engine/Source/Runtime/Core/Tests/Logging/StructuredLogFormatTest.cpp b/Engine/Source/Runtime/Core/Tests/Logging/StructuredLogFormatTest.cpp
--- a/Engine/Source/Runtime/Core/Tests/Logging/StructuredLogFormatTest.cpp
+++ b/Engine/Source/Runtime/Core/Tests/Logging/StructuredLogFormatTest.cpp
@@ -242,6 +242,17 @@ TEST_CASE_NAMED(FStructuredLogFormatTest, "Core::Logging::StructuredLogFormat",
 	Test(TEXT("{ObjectText/X}"), ANSITEXTVIEW("1"), {.bAllowSubObjectReferences = true});
 	Test(TEXT("{ObjectWithNestedText/X}"), ANSITEXTVIEW("0001"), {.bAllowSubObjectReferences = true});
 
+	// Test array indexing
+	Test(TEXT("{Array[0]}"), ANSITEXTVIEW(R"(1)"), {.bAllowSubObjectReferences = true});
+	Test(TEXT("{Array[1]}"), ANSITEXTVIEW(R"(2)"), {.bAllowSubObjectReferences = true});
+	Test(TEXT("{Array[2]}"), ANSITEXTVIEW(R"(3)"), {.bAllowSubObjectReferences = true});
+	Test(TEXT("{ArrayWithNestedFormat[0]}"), ANSITEXTVIEW(R"(X=1)"), {.bAllowSubObjectReferences = true});
+	Test(TEXT("{ArrayWithNestedFormat[1]}"), ANSITEXTVIEW(R"(Y=2)"), {.bAllowSubObjectReferences = true});
+	Test(TEXT("{ArrayWithNestedFormat[2]}"), ANSITEXTVIEW(R"(Z=3)"), {.bAllowSubObjectReferences = true});
+	Test(TEXT("{ArrayWithNestedFormat[0]/X}"), ANSITEXTVIEW(R"(1)"), {.bAllowSubObjectReferences = true});
+	Test(TEXT("{ArrayWithNestedFormat[1]/Y}"), ANSITEXTVIEW(R"(2)"), {.bAllowSubObjectReferences = true});
+	Test(TEXT("{ArrayWithNestedFormat[2]/Z}"), ANSITEXTVIEW(R"(3)"), {.bAllowSubObjectReferences = true});
+
 	const auto TestLoc = [&Fields](const FText& FormatText, FUtf8StringView Expected, const FLogTemplateOptions& Options = {})
 	{
 		const TOptional<FString> Namespace = FTextInspector::GetNamespace(FormatText);
@@ -308,6 +319,17 @@ TEST_CASE_NAMED(FStructuredLogFormatTest, "Core::Logging::StructuredLogFormat",
 	TestLoc(LOCTEXT("ObjectTextX", "{ObjectText/X}"), ANSITEXTVIEW("1"), {.bAllowSubObjectReferences = true});
 	TestLoc(LOCTEXT("ObjectWithNestedTextX", "{ObjectWithNestedText/X}"), ANSITEXTVIEW("0001"), {.bAllowSubObjectReferences = true});
 
+	// Test array indexing
+	TestLoc(LOCTEXT("ArrayIndex0", "{Array[0]}"), ANSITEXTVIEW(R"(1)"), {.bAllowSubObjectReferences = true});
+	TestLoc(LOCTEXT("ArrayIndex1", "{Array[1]}"), ANSITEXTVIEW(R"(2)"), {.bAllowSubObjectReferences = true});
+	TestLoc(LOCTEXT("ArrayIndex2", "{Array[2]}"), ANSITEXTVIEW(R"(3)"), {.bAllowSubObjectReferences = true});
+	TestLoc(LOCTEXT("ArrayWithNestedFormatIndex0", "{ArrayWithNestedFormat[0]}"), ANSITEXTVIEW(R"(X=1)"), {.bAllowSubObjectReferences = true});
+	TestLoc(LOCTEXT("ArrayWithNestedFormatIndex1", "{ArrayWithNestedFormat[1]}"), ANSITEXTVIEW(R"(Y=2)"), {.bAllowSubObjectReferences = true});
+	TestLoc(LOCTEXT("ArrayWithNestedFormatIndex2", "{ArrayWithNestedFormat[2]}"), ANSITEXTVIEW(R"(Z=3)"), {.bAllowSubObjectReferences = true});
+	TestLoc(LOCTEXT("ArrayWithNestedFormatIndex0X", "{ArrayWithNestedFormat[0]/X}"), ANSITEXTVIEW(R"(1)"), {.bAllowSubObjectReferences = true});
+	TestLoc(LOCTEXT("ArrayWithNestedFormatIndex1Y", "{ArrayWithNestedFormat[1]/Y}"), ANSITEXTVIEW(R"(2)"), {.bAllowSubObjectReferences = true});
+	TestLoc(LOCTEXT("ArrayWithNestedFormatIndex2Z", "{ArrayWithNestedFormat[2]/Z}"), ANSITEXTVIEW(R"(3)"), {.bAllowSubObjectReferences = true});
+
 	// Test constructing a template from FText.
 	{
 		FUniqueLogTemplate Template(LOCTEXT("TextFormat", "FText bWorks={True}"));
