Fixes for Clang 21 -Wcharacter-conversion warnings across UE char-type templates

Resolved implicit conversions between char8_t, char16_t, and char32_t in template code that operates on mixed UE char types (ANSICHAR, UTF8CHAR, WIDECHAR) by adding explicit casts and using already templatized versions of existing APIs where it was possible to do so.

#Synced-CL 51857304

diff --git a/Engine/Source/Developer/TraceInsights/Private/Insights/Common/Utf8FileWriter.cpp b/Engine/Source/Developer/TraceInsights/Private/Insights/Common/Utf8FileWriter.cpp
--- a/Engine/Source/Developer/TraceInsights/Private/Insights/Common/Utf8FileWriter.cpp
+++ b/Engine/Source/Developer/TraceInsights/Private/Insights/Common/Utf8FileWriter.cpp
@@ -100,7 +100,7 @@ void FUtf8FileWriter::AppendValue(FStringView Str)
 	if (IsCSV())
 	{
 		bool bNeedsQuotes = false;
-		if (UE::String::FindFirstChar(Str, Separator, ESearchCase::CaseSensitive) != INDEX_NONE)
+		if (UE::String::FindFirstChar(Str, static_cast<TCHAR>(Separator), ESearchCase::CaseSensitive) != INDEX_NONE)
 		{
 			AppendDoubleQuotationMark();
 			bNeedsQuotes = true;
diff --git a/Engine/Source/Runtime/Core/Private/GenericPlatform/GenericPlatformString.cpp b/Engine/Source/Runtime/Core/Private/GenericPlatform/GenericPlatformString.cpp
--- a/Engine/Source/Runtime/Core/Private/GenericPlatform/GenericPlatformString.cpp
+++ b/Engine/Source/Runtime/Core/Private/GenericPlatform/GenericPlatformString.cpp
@@ -1368,14 +1368,14 @@ int32 StrncmpImpl(const CharType1* String1, const CharType2* String2, SIZE_T Cou
 {
 	for (; Count > 0; --Count)
 	{
-		CharType1 C1 = *String1++;
-		CharType2 C2 = *String2++;
+		uint32 AsUnsignedChar1 = TChar<CharType1>::ToUnsigned(*String1++);
+		uint32 AsUnsignedChar2 = TChar<CharType2>::ToUnsigned(*String2++);
 
-		if (C1 != C2)
+		if (AsUnsignedChar1 != AsUnsignedChar2)
 		{
-			return TChar<CharType1>::ToUnsigned(C1) - TChar<CharType2>::ToUnsigned(C2);
+			return AsUnsignedChar1 - AsUnsignedChar2;
 		}
-		if (C1 == 0)
+		if (AsUnsignedChar1 == 0)
 		{
 			return 0;
 		}
diff --git a/Engine/Source/Runtime/Core/Private/Misc/ExpressionParser.cpp b/Engine/Source/Runtime/Core/Private/Misc/ExpressionParser.cpp
--- a/Engine/Source/Runtime/Core/Private/Misc/ExpressionParser.cpp
+++ b/Engine/Source/Runtime/Core/Private/Misc/ExpressionParser.cpp
@@ -61,13 +61,13 @@ FString TTokenStream<CharType>::GetErrorContext() const
 	const CharType* EndPos = StartPos;
 
 	// Skip over any leading whitespace
-	while (FChar::IsWhitespace(*EndPos))
+	while (TChar<CharType>::IsWhitespace(*EndPos))
 	{
 		EndPos++;
 	}
 
 	// Read until next whitespace or end of string
-	while (!FChar::IsWhitespace(*EndPos) && *EndPos != '\0')
+	while (!TChar<CharType>::IsWhitespace(*EndPos) && *EndPos != '\0')
 	{
 		EndPos++;
 	}
diff --git a/Engine/Source/Runtime/Core/Private/Misc/Paths.cpp b/Engine/Source/Runtime/Core/Private/Misc/Paths.cpp
--- a/Engine/Source/Runtime/Core/Private/Misc/Paths.cpp
+++ b/Engine/Source/Runtime/Core/Private/Misc/Paths.cpp
@@ -79,8 +79,8 @@ struct FPaths::FStaticData
 
 namespace UE4Paths_Private
 {
-	auto IsSlashOrBackslash    = [](TCHAR C) { return C == TEXT('/') || C == TEXT('\\'); };
-	auto IsNotSlashOrBackslash = [](TCHAR C) { return C != TEXT('/') && C != TEXT('\\'); };
+	auto IsSlashOrBackslash    = [](auto C) { return C == '/' || C == '\\'; };
+	auto IsNotSlashOrBackslash = [](auto C) { return C != '/' && C != '\\'; };
 
 	FString GetSavedDirSuffix(const FString& BaseDir, const TCHAR* CommandLineArgument)
 	{
diff --git a/Engine/Source/Runtime/Core/Private/Misc/StringFormatter.cpp b/Engine/Source/Runtime/Core/Private/Misc/StringFormatter.cpp
--- a/Engine/Source/Runtime/Core/Private/Misc/StringFormatter.cpp
+++ b/Engine/Source/Runtime/Core/Private/Misc/StringFormatter.cpp
@@ -175,12 +175,12 @@ TOptional<FExpressionError> ParseIndex(TExpressionTokenConsumer<CharType>& Consu
 	TStringToken<CharType>& EntireToken = OpeningChar.GetValue();
 
 	// Optional whitespace
-	Stream.ParseToken([](TCHAR InC) { return FChar::IsWhitespace(InC) ? EParseState::Continue : EParseState::StopBefore; }, &EntireToken);
+	Stream.ParseToken([](CharType InC) { return TChar<CharType>::IsWhitespace(InC) ? EParseState::Continue : EParseState::StopBefore; }, &EntireToken);
 
 	// The identifier itself
 	TOptional<int32> Index;
-	Stream.ParseToken([&](TCHAR InC) {
-		if (FChar::IsDigit(InC))
+	Stream.ParseToken([&](CharType InC) {
+		if (TChar<CharType>::IsDigit(InC))
 		{
 			if (!Index.IsSet())
 			{
@@ -207,7 +207,7 @@ TOptional<FExpressionError> ParseIndex(TExpressionTokenConsumer<CharType>& Consu
 	}
 
 	// Optional whitespace
-	Stream.ParseToken([](TCHAR InC) { return FChar::IsWhitespace(InC) ? EParseState::Continue : EParseState::StopBefore; }, &EntireToken);
+	Stream.ParseToken([](CharType InC) { return TChar<CharType>::IsWhitespace(InC) ? EParseState::Continue : EParseState::StopBefore; }, &EntireToken);
 	
 	if (!Stream.ParseSymbol(CHARTEXT(CharType, '}'), &EntireToken).IsSet())
 	{
diff --git a/Engine/Source/Runtime/Core/Private/Serialization/CompactBinarySerialization.cpp b/Engine/Source/Runtime/Core/Private/Serialization/CompactBinarySerialization.cpp
--- a/Engine/Source/Runtime/Core/Private/Serialization/CompactBinarySerialization.cpp
+++ b/Engine/Source/Runtime/Core/Private/Serialization/CompactBinarySerialization.cpp
@@ -114,17 +114,17 @@ void AppendQuotedJsonString(TStringBuilderBase<CharType>& Builder, FUtf8StringVi
 		Builder << Verbatim;
 		Value.RightChopInline(Verbatim.Len());
 		FUtf8StringView Escape = FAsciiSet::FindPrefixWith(Value, EscapeSet);
-		for (CharType Char : Escape)
+		for (UTF8CHAR Char : Escape)
 		{
 			switch (Char)
 			{
-			case '\\': Builder << FConstants::EscapeBackslash; break;
-			case '\"': Builder << FConstants::EscapeDoubleQuote; break;
-			case '\b': Builder << FConstants::EscapeBell; break;
-			case '\f': Builder << FConstants::EscapeFormFeed; break;
-			case '\n': Builder << FConstants::EscapeLineFeed; break;
-			case '\r': Builder << FConstants::EscapeCarriageReturn; break;
-			case '\t': Builder << FConstants::EscapeTab; break;
+			case UTF8CHAR('\\'): Builder << FConstants::EscapeBackslash; break;
+			case UTF8CHAR('\"'): Builder << FConstants::EscapeDoubleQuote; break;
+			case UTF8CHAR('\b'): Builder << FConstants::EscapeBell; break;
+			case UTF8CHAR('\f'): Builder << FConstants::EscapeFormFeed; break;
+			case UTF8CHAR('\n'): Builder << FConstants::EscapeLineFeed; break;
+			case UTF8CHAR('\r'): Builder << FConstants::EscapeCarriageReturn; break;
+			case UTF8CHAR('\t'): Builder << FConstants::EscapeTab; break;
 			default:
 				Builder.Appendf(FConstants::FormatEscapeUnicode4, uint32(Char));
 				break;
diff --git a/Engine/Source/Runtime/Core/Private/String/HexToBytes.cpp b/Engine/Source/Runtime/Core/Private/String/HexToBytes.cpp
--- a/Engine/Source/Runtime/Core/Private/String/HexToBytes.cpp
+++ b/Engine/Source/Runtime/Core/Private/String/HexToBytes.cpp
@@ -17,12 +17,12 @@ static inline int32 HexToBytesImpl(TStringView<CharType> Hex, uint8* const OutBy
 	uint8* OutPos = OutBytes;
 	if (const bool bPadNibble = (HexCount % 2) == 1)
 	{
-		*OutPos++ = TCharToNibble(*HexPos++);
+		*OutPos++ = TCharToNibble(static_cast<TCHAR>(*HexPos++));
 	}
 	while (HexPos != HexEnd)
 	{
-		const uint8 HiNibble = uint8(TCharToNibble(*HexPos++) << 4);
-		*OutPos++ = HiNibble | TCharToNibble(*HexPos++);
+		const uint8 HiNibble = uint8(TCharToNibble(static_cast<TCHAR>(*HexPos++)) << 4);
+		*OutPos++ = HiNibble | TCharToNibble(static_cast<TCHAR>(*HexPos++));
 	}
 	return static_cast<int32>(OutPos - OutBytes);
 }
diff --git a/Engine/Source/Runtime/Core/Private/String/Numeric.cpp b/Engine/Source/Runtime/Core/Private/String/Numeric.cpp
--- a/Engine/Source/Runtime/Core/Private/String/Numeric.cpp
+++ b/Engine/Source/Runtime/Core/Private/String/Numeric.cpp
@@ -34,7 +34,7 @@ static bool IsNumeric(TStringView<CharType> View)
 
 			bHasDot = true;
 		}
-		else if (!FChar::IsDigit(*CurPtr))
+		else if (!TChar<CharType>::IsDigit(*CurPtr))
 		{
 			return false;
 		}
@@ -58,7 +58,7 @@ static bool IsNumericOnlyDigits(TStringView<CharType> View)
 
 	while (CurPtr < EndPtr)
 	{
-		if (!FChar::IsDigit(*CurPtr))
+		if (!TChar<CharType>::IsDigit(*CurPtr))
 		{
 			return false;
 		}
diff --git a/Engine/Source/Runtime/Core/Public/Misc/CString.h b/Engine/Source/Runtime/Core/Public/Misc/CString.h
--- a/Engine/Source/Runtime/Core/Public/Misc/CString.h
+++ b/Engine/Source/Runtime/Core/Public/Misc/CString.h
@@ -163,7 +163,7 @@ struct TCString
 				}
 				bHasDot = true;
 			}
-			else if (!FChar::IsDigit(*Str))
+			else if (!TChar<CharType>::IsDigit(*Str))
 			{
 				return false;
 			}
diff --git a/Engine/Source/Runtime/Json/Public/Serialization/JsonReader.h b/Engine/Source/Runtime/Json/Public/Serialization/JsonReader.h
--- a/Engine/Source/Runtime/Json/Public/Serialization/JsonReader.h
+++ b/Engine/Source/Runtime/Json/Public/Serialization/JsonReader.h
@@ -701,7 +701,7 @@ private:
 							}
 							++CharacterNumber;
 
-							int32 HexDigit = FParse::HexDigit(Char);
+							int32 HexDigit = FParse::HexDigit(static_cast<TCHAR>(Char));
 
 							if ((HexDigit == 0) && (Char != CharType('0')))
 							{
diff --git a/Engine/Source/Runtime/Online/HTTP/Private/GenericPlatform/GenericPlatformHttp.cpp b/Engine/Source/Runtime/Online/HTTP/Private/GenericPlatform/GenericPlatformHttp.cpp
--- a/Engine/Source/Runtime/Online/HTTP/Private/GenericPlatform/GenericPlatformHttp.cpp
+++ b/Engine/Source/Runtime/Online/HTTP/Private/GenericPlatform/GenericPlatformHttp.cpp
@@ -398,10 +398,10 @@ FString FGenericPlatformHttp::UrlDecode(const FStringView EncodedString)
 				if (CharIdx <= Converter.Length() - 6)
 				{
 					// Treat all %uXXXX as code point
-					Value = FParse::HexDigit(UTF8Data[CharIdx + 2]) << 12;
-					Value += FParse::HexDigit(UTF8Data[CharIdx + 3]) << 8;
-					Value += FParse::HexDigit(UTF8Data[CharIdx + 4]) << 4;
-					Value += FParse::HexDigit(UTF8Data[CharIdx + 5]);
+					Value = FParse::HexDigit(static_cast<TCHAR>(UTF8Data[CharIdx + 2])) << 12;
+					Value += FParse::HexDigit(static_cast<TCHAR>(UTF8Data[CharIdx + 3])) << 8;
+					Value += FParse::HexDigit(static_cast<TCHAR>(UTF8Data[CharIdx + 4])) << 4;
+					Value += FParse::HexDigit(static_cast<TCHAR>(UTF8Data[CharIdx + 5]));
 					CharIdx += 6;
 
 					UTF8CHAR Buffer[8] = {};
@@ -422,8 +422,8 @@ FString FGenericPlatformHttp::UrlDecode(const FStringView EncodedString)
 			else if(CharIdx <= Converter.Length() - 3)
 			{
 				// Treat all %XX as straight byte
-				Value = FParse::HexDigit(UTF8Data[CharIdx + 1]) << 4;
-				Value += FParse::HexDigit(UTF8Data[CharIdx + 2]);
+				Value = FParse::HexDigit(static_cast<TCHAR>(UTF8Data[CharIdx + 1])) << 4;
+				Value += FParse::HexDigit(static_cast<TCHAR>(UTF8Data[CharIdx + 2]));
 				CharIdx += 3;
 				Data.Add((UTF8CHAR)(Value));
 			}
