UE: Fix loading cooked content that contains editoronly data

#Synced-CL 51809058

diff --git a/Engine/Source/Runtime/Core/Private/Internationalization/GatherableTextData.cpp b/Engine/Source/Runtime/Core/Private/Internationalization/GatherableTextData.cpp
--- a/Engine/Source/Runtime/Core/Private/Internationalization/GatherableTextData.cpp
+++ b/Engine/Source/Runtime/Core/Private/Internationalization/GatherableTextData.cpp
@@ -26,7 +26,7 @@ void FTextSourceSiteContext::StreamStructuredArchive(FStructuredArchive::FSlot S
 #if WITH_EDITORONLY_DATA
 	FArchive& UnderlyingArchive = Slot.GetUnderlyingArchive();
 	UnderlyingArchive.UsingCustomVersion(FFortniteMainBranchObjectVersion::GUID);
-	if (UnderlyingArchive.CustomVer(FFortniteMainBranchObjectVersion::GUID) >= FFortniteMainBranchObjectVersion::AddDevNotesToFText && !UnderlyingArchive.IsCooking() && !UnderlyingArchive.IsFilterEditorOnly())
+	if (UnderlyingArchive.CustomVer(FFortniteMainBranchObjectVersion::GUID) >= FFortniteMainBranchObjectVersion::AddDevNotesToFText && !UnderlyingArchive.IsFilterEditorOnly())
 	{
 		Record << SA_VALUE(TEXT("DevNotes"), DevNotes);
 	}
diff --git a/Engine/Source/Runtime/Core/Private/Internationalization/TextHistory.cpp b/Engine/Source/Runtime/Core/Private/Internationalization/TextHistory.cpp
--- a/Engine/Source/Runtime/Core/Private/Internationalization/TextHistory.cpp
+++ b/Engine/Source/Runtime/Core/Private/Internationalization/TextHistory.cpp
@@ -914,7 +914,7 @@ void FTextHistory_Base::Serialize(FStructuredArchive::FRecord Record)
 
 #if WITH_EDITORONLY_DATA
 	BaseArchive.UsingCustomVersion(FFortniteMainBranchObjectVersion::GUID);
-	if (BaseArchive.CustomVer(FFortniteMainBranchObjectVersion::GUID) >= FFortniteMainBranchObjectVersion::AddDevNotesToFText && !BaseArchive.IsCooking() && !BaseArchive.IsFilterEditorOnly())
+	if (BaseArchive.CustomVer(FFortniteMainBranchObjectVersion::GUID) >= FFortniteMainBranchObjectVersion::AddDevNotesToFText && !BaseArchive.IsFilterEditorOnly())
 	{
 		Record << SA_VALUE(TEXT("DevNotes"), DevNotes);
 	}
