Early access configrules.

These changes make configrule variables available very early, such that they can be used inside static initialisers, making it useable to select the malloc allocator for example.
- Java code creates a locally named socket and waits for native code to read. A peer UID check ensures only the app's socket is accepted.
- When native accesses configrules it connects to the socket and transfers a binary blob of configrule K/V pairs.
- nativeSetConfigRulesVariables also passes the binary blob so if any socket communication failures occur the API still functions.
- Added GetConfigRuleVariableAsUtf8, an alloc free lookup into early config rules binary blob.
- Java side getGraphicsInformation is called before ProcessSystemInfo, it is possible for it to wait for .so loader which will deadlock if the engine is waiting for configrule info inside its own loadlibrary call.
- Changed to use NDK supplied ATrace calls directly instead of dlsym lookups.
- Add an example use case for selecting mimalloc allocator via config rules.

#Synced-CL 51487278

diff --git a/Engine/Build/Android/Java/src/com/epicgames/unreal/GameActivity.java.template b/Engine/Build/Android/Java/src/com/epicgames/unreal/GameActivity.java.template
--- a/Engine/Build/Android/Java/src/com/epicgames/unreal/GameActivity.java.template
+++ b/Engine/Build/Android/Java/src/com/epicgames/unreal/GameActivity.java.template
@@ -140,6 +140,8 @@ import android.content.SharedPreferences.Editor;
 import android.provider.Settings;
 import android.provider.Settings.Secure;
 import android.content.pm.PackageInfo;
+import android.net.LocalServerSocket;
+import android.net.LocalSocket;
 
 import android.webkit.CookieManager;
 import android.webkit.CookieSyncManager;
@@ -373,7 +375,6 @@ public class GameActivity extends $${gameActivitySuperClass}$$ implements Surfac
 
 	protected Thread ProcessSystemInfoThread = null;
 	private Lock SystemInfoLogsLock = new ReentrantLock();
-	private Lock ProcessSystemInfoLock = new ReentrantLock();
 	protected boolean ConfigRulesRunResult = true;
 	protected java.security.Key ConfigRulesKey = null;
 	protected java.security.Key ConfigRulesCacheKey = null;
@@ -2422,7 +2423,7 @@ Trace.endSection();
 		}
 	}
 
-	protected void processSystemInfo(Lock lock)
+	protected void processSystemInfo()
 	{
 		SystemInfo_InitLog();
 
@@ -2591,25 +2592,12 @@ Trace.endSection();
 		variables.put("screenWidth", Integer.toString(displaySize.x));
 		variables.put("screenHeight", Integer.toString(displaySize.y));
 
-		// check OpenGL information
-		if (lock != null)
-		{
-Trace.beginSection("GLLock");
-			lock.lock();
-Trace.endSection();
-		}
-Trace.beginSection("getGraphicsInfo");
-		getGraphicsInformation();
+		// get OpenGL information
 		variables.put("SRC_GpuFamily", OpenGLDevice);
 		variables.put("SRC_GlVersion", OpenGLDriver);
 		variables.put("SRC_GpuVendor", OpenGLVendor);
 		variables.put("supportsFloatingPointRenderTargets", bSupportsFloatingPointRenderTargets ? "true" : "false");
 		variables.put("TextureFormats", TextureFormats);
-Trace.endSection();
-		if (lock != null)
-		{
-			lock.unlock();
-		}
 
 Trace.beginSection("ReadCRFiles");
 		ConfigRulesBytesAssets = getByteArrayFromAssets("configrules.bin.png");
@@ -3638,14 +3626,7 @@ Trace.endSection();  // WriteCache
 		}
 
 		// send the variables to the engine
-		String[] keyValues = new String[variables.size() * 2];
-		int keyValueIndex = 0;
-		for (Map.Entry<String, String> entry : variables.entrySet())
-		{
-			keyValues[keyValueIndex++] = entry.getKey();
-			keyValues[keyValueIndex++] = entry.getValue();
-		}
-		nativeSetConfigRulesVariables(keyValues);
+		SetNativeConfigRuleVariables(variables);
 		bConfigRulesApplied = true;
 	}
 
@@ -3723,8 +3704,104 @@ Trace.endSection();  // WriteCache
 		return null;
 	}
 
+	private static class EarlyConfigrulesSetup implements Runnable
+	{
+		private static LocalServerSocket ConfigRulesSocket;
 
+		private static byte[] ConfigrulesBytes = null;
 
+		public EarlyConfigrulesSetup(Map<String, String> VarsIn)
+		{
+			// prepare the variables as a single binary blob ready to send to the engine.
+			ByteArrayOutputStream ConfigrulesByteStream = new ByteArrayOutputStream();
+			try {
+				for (Map.Entry<String, String> entry : VarsIn.entrySet())
+				{
+					byte[] KeyBytes = entry.getKey().getBytes(StandardCharsets.UTF_8);
+					byte[] ValueBytes = entry.getValue().getBytes(StandardCharsets.UTF_8);
+					if(KeyBytes.length< 256 && ValueBytes.length < 256)
+					{
+						ConfigrulesByteStream.write(KeyBytes.length);
+						ConfigrulesByteStream.write(KeyBytes);
+						ConfigrulesByteStream.write(ValueBytes.length);
+						ConfigrulesByteStream.write(ValueBytes);
+					}
+					else
+					{
+						Log.warn("Skipped long configrule string in kv pair: " + entry.getKey() + ":" + entry.getValue());
+					}
+				}
+			}
+			catch (IOException e)
+			{
+				Log.debug("failed to get early configrule byte stream: " + e.getMessage());
+				ConfigrulesByteStream.reset();
+			}
+			ConfigrulesBytes = ConfigrulesByteStream.toByteArray();
+		}
+
+		@Override
+		public void run()
+		{
+			try
+			{
+Trace.beginSection("earlyConfigRules socket ready");
+				ConfigRulesSocket = new LocalServerSocket("ue.earlyconfigrules");
+				try (LocalSocket Sock = ConfigRulesSocket.accept())
+				{
+					// Abstract sockets are globally visible. Check peer UID matches ours to reject
+					// other apps. Processes sharing our UID (e.g. child services) already have
+					// equivalent file access so this is the appropriate trust boundary.
+					int peerUid = Sock.getPeerCredentials().getUid();
+					if (peerUid != android.os.Process.myUid())
+					{
+						Log.warn("EarlyConfigrulesSetup: Rejecting connection from UID " + peerUid + " (expected " + android.os.Process.myUid() + ")");
+						return;
+					}
+					OutputStream Out = Sock.getOutputStream();
+					Out.write(ConfigrulesBytes);
+					Out.flush();
+				}
+			}
+			catch(Exception e)
+			{
+				Log.debug("EarlyConfigrulesSetup: exception while waiting for socket accept: " + e.getMessage());
+			}
+			finally
+			{
+				if (ConfigRulesSocket != null)
+				{
+					try
+					{
+						ConfigRulesSocket.close();
+					}
+					catch (Exception e)
+					{
+					}
+					ConfigRulesSocket = null;
+				}
+Trace.endSection();
+			}
+		}
+
+		public static void Shutdown()
+		{
+			if(ConfigRulesSocket != null)
+			{
+				try
+				{
+					Log.debug("EarlyConfigrulesSetup shutdown - closing socket.");
+					ConfigRulesSocket.close();
+					ConfigRulesSocket = null;
+				}
+				catch(Exception e)
+				{
+					Log.debug("EarlyConfigrulesSetup shutdown Error closing socket: " + e.getMessage());
+				}
+			}
+			ConfigrulesBytes = null;
+		}
+	}
 
 	@Override
 	public void onCreate(Bundle savedInstanceState)
@@ -3751,13 +3828,6 @@ Trace.endSection();
 		MainThreadId = android.os.Process.myTid();
 		MainThreadPriority = android.os.Process.getThreadPriority(MainThreadId);
 
-		// block getGraphicsInformation on ProcessSystemInfoThread until GLES loads
-		// no need if using SplashActivity as it uses the GameApplication lock instead
-		if (!HAS_SPLASH_ACTIVITY)
-		{
-			ProcessSystemInfoLock.lock();
-		}
-
 //$${gameActivityOnCreateBeforeProcessSystemInfoAdditions}$$
 		ProcessSystemInfoThread = new Thread(new Runnable() {
 			public void run() {
@@ -3766,14 +3836,7 @@ Trace.endSection();
 				ParseCommandline(ProjectName, UseExternalFilesDir);
 Trace.beginSection("processSystemInfo");
 				// collect system info
-				if (HAS_SPLASH_ACTIVITY)
-				{
-					processSystemInfo(com.epicgames.unreal.GameApplication.ProcessSystemInfoLock);
-				}
-				else
-				{
-					processSystemInfo(ProcessSystemInfoLock);
-				}
+				processSystemInfo();
 Trace.endSection();
 
 				// run config rules
@@ -3784,29 +3847,34 @@ Trace.beginSection("runConfigRules");
 				Log.debug("$$$ configrules: run completed");
 Trace.endSection();
 //$${gameActivityProcessSystemInfoThreadAfterConfigRulesAdditions}$$
+Trace.beginSection("earlyConfigRules serialize");
+				EarlyConfigrulesSetup earlyConfigrulesSetup = new EarlyConfigrulesSetup(ConfigRulesVars);
+Trace.endSection();
+				Thread earlyConfigrulesThread = new Thread(earlyConfigrulesSetup);
+				earlyConfigrulesThread.setName("EarlyCRSocket");
+				earlyConfigrulesThread.start();
 			}
 		});
-		ProcessSystemInfoThread.setName("ProcessSysInfo");
-		ProcessSystemInfoThread.start();
 
-		// load GLES/EGL now so available for getGraphicsInformation in ProcessSystemInfo thread without RenderDoc hang
-		// cannot be in another thread as it will just block (fairly quick to load so not a big deal to do it here)
-		// SplashActivity would have already tried to load early so can skip if we have it
-		if (!HAS_SPLASH_ACTIVITY)
+		try
 		{
-Trace.beginSection("loadGLES");
-			try
-			{
-				System.loadLibrary("GLESv3");
-				System.loadLibrary("EGL");
-			}
-			catch (Exception e)
-			{
-			}
-			// unlock so getGraphicsInformation is allowed to run
-			ProcessSystemInfoLock.unlock();
+Trace.beginSection("preload getGraphicsInfo for ProcessSystemInfo");
+			// getGraphicsInformation has been moved from ProcessSystemInfo. 
+			// getGraphicsInformation may trigger loadLibrary calls, if that happens during the engine's loadlibrary call it must wait for the engine library to finish.
+			// if an early configrule access is made it will timeout as it must wait for getGraphicsInformation which cannot complete.
+			// previously this could cause a deadlock.
+			// 
+			// Prepare the gfx info now so confrules will not wait for engine's loadlibrary.
+			getGraphicsInformation();
 Trace.endSection();
 		}
+		catch (Exception e)
+		{
+			Log.debug("getGraphicsInformation threw " + e.getMessage());
+		}
+
+		ProcessSystemInfoThread.setName("ProcessSysInfo");
+		ProcessSystemInfoThread.start();
 
 //$${gameActivityOnCreateBeforeLoadLibrariesAdditions}$$
 
@@ -3847,8 +3915,9 @@ Trace.endSection();
 			SystemInfo_DumpLog();
 			// the native code will wait for configrules variables.
 			// something must be sent to release it.
-			nativeSetConfigRulesVariables(new String[]{});
+			SetNativeConfigRuleVariables(new HashMap<String,String>());
 		}
+
 	}
 
 	protected void onCreatePreBody()
@@ -8755,10 +8824,30 @@ Trace.endSection();
 //$${gameActivitySetCrashContextData}$$
 	}
 
+	public void SetNativeConfigRuleVariables(Map<String, String> VarsIn)
+	{
+Trace.beginSection("SetNativeConfigRuleVariables");
+		// End the socket side attempt, we will use JNI to transfer the data now.
+		// Any current native side Socket recieve attempt will abort and behave as a non-existing key.
+		byte[] ConfigrulesBytes = EarlyConfigrulesSetup.ConfigrulesBytes;
+		EarlyConfigrulesSetup.Shutdown();
+
+		// send the variables to the engine
+		String[] keyValues = new String[VarsIn.size() * 2];
+		int keyValueIndex = 0;
+		for (Map.Entry<String, String> entry : VarsIn.entrySet())
+		{
+			keyValues[keyValueIndex++] = entry.getKey();
+			keyValues[keyValueIndex++] = entry.getValue();
+		}
+		nativeSetConfigRulesVariables(keyValues, ConfigrulesBytes);
+Trace.endSection();
+	}
+
 	public native int nativeGetCPUFamily();
 	public native boolean nativeSupportsNEON();
 	public native void nativeSetAffinityInfo(boolean bEnableAffinity, int bigCoreMask, int littleCoreMask);
-	public native void nativeSetConfigRulesVariables(String[] KeyValuePairs);
+	public native void nativeSetConfigRulesVariables(String[] KeyValuePairs, byte[] configrulesBytes);
 
 	public native boolean nativeIsShippingBuild();
 	public native void nativeSetAndroidStartupState(boolean bDebuggerAttached);
diff --git a/Engine/Source/Runtime/Core/Private/Android/AndroidPlatformMemory.cpp b/Engine/Source/Runtime/Core/Private/Android/AndroidPlatformMemory.cpp
--- a/Engine/Source/Runtime/Core/Private/Android/AndroidPlatformMemory.cpp
+++ b/Engine/Source/Runtime/Core/Private/Android/AndroidPlatformMemory.cpp
@@ -469,8 +469,8 @@ FMalloc* FAndroidPlatformMemory::BaseAllocator()
 #endif
 
 #if	!ENABLE_PGO_PROFILE
-	const bool bHasEnoughMemoryForMimalloc = false;	//stub
-	if (bHasEnoughMemoryForMimalloc)
+	const bool bUseMimalloc = FPlatformMisc::GetConfigRulesVariableAsUtf8("malloctype").Equals("mimalloc");
+	if (bUseMimalloc)
 	{
 		return (Instance = new FMallocMimalloc());
 	}
diff --git a/Engine/Source/Runtime/Core/Private/Android/AndroidPlatformMisc.cpp b/Engine/Source/Runtime/Core/Private/Android/AndroidPlatformMisc.cpp
--- a/Engine/Source/Runtime/Core/Private/Android/AndroidPlatformMisc.cpp
+++ b/Engine/Source/Runtime/Core/Private/Android/AndroidPlatformMisc.cpp
@@ -19,6 +19,8 @@
 #include <sys/statfs.h>
 #include <sys/syscall.h>
 #include <sched.h>
+#include <sys/socket.h>
+#include <sys/un.h>
 
 #include "Android/AndroidPlatformCrashContext.h"
 #include "Android/AndroidStats.h"
@@ -27,8 +29,10 @@
 #include "GenericPlatform/GenericPlatformChunkInstall.h"
 
 #include "Misc/Parse.h"
+#include "Misc/ScopeExit.h"
 #include "Internationalization/Regex.h"
 
+#include <android/trace.h>
 #include <android/log.h>
 #if USE_ANDROID_INPUT
 #include <android/keycodes.h>
@@ -135,16 +139,6 @@ static FAutoConsoleVariableRef CVarAndroidBroadcastIntentData(
 #if STATS || ENABLE_STATNAMEDEVENTS
 int32 FAndroidMisc::TraceMarkerFileDescriptor = -1;
 
-typedef void (*ATrace_beginSection_Type) (const char* sectionName);
-typedef void (*ATrace_endSection_Type) (void);
-typedef bool (*ATrace_isEnabled_Type) (void);
-
-static ATrace_beginSection_Type ATrace_beginSection = NULL;
-static ATrace_endSection_Type ATrace_endSection = NULL;
-static ATrace_isEnabled_Type ATrace_isEnabled = NULL;
-
-static bool bUseNativeSystrace = false;
-
 #endif
 
 // run time compatibility information
@@ -602,43 +596,6 @@ void FAndroidMisc::PlatformInit()
 	extern void AndroidSetupDefaultThreadAffinity();
 	AndroidSetupDefaultThreadAffinity();
 
-#if (STATS || ENABLE_STATNAMEDEVENTS)
-	//Loading NDK libandroid.so atrace functions, available in the android libraries way before NDK headers.
-	void* const LibAndroid = dlopen("libandroid.so", RTLD_NOW | RTLD_LOCAL);
-	if (LibAndroid != nullptr)
-	{
-		// Retrieve function pointers from shared object.
-		ATrace_beginSection = reinterpret_cast<ATrace_beginSection_Type>(dlsym(LibAndroid, "ATrace_beginSection"));
-		ATrace_endSection = reinterpret_cast<ATrace_endSection_Type>(dlsym(LibAndroid, "ATrace_endSection"));
-		ATrace_isEnabled = 	reinterpret_cast<ATrace_isEnabled_Type>(dlsym(LibAndroid, "ATrace_isEnabled"));
-	}
-
-	if (!ATrace_beginSection || !ATrace_endSection || !ATrace_isEnabled)
-	{
-		UE_LOG(LogAndroid, Warning, TEXT("Failed to use native systrace functionality."));
-		ATrace_beginSection = nullptr;
-		ATrace_endSection = nullptr;
-		ATrace_isEnabled = nullptr;
-
-		if (FParse::Param(FCommandLine::Get(), TEXT("enablesystrace")))
-		{
-			GAndroidTraceMarkersEnabled = 1;
-		}
-
-		if (GAndroidTraceMarkersEnabled)
-		{
-			StartTraceMarkers();
-		}
-
-		// Watch for CVar update
-		CAndroidTraceMarkersEnabled->SetOnChangedCallback(FConsoleVariableDelegate::CreateStatic(&UpdateTraceMarkersEnable));
-	}
-	else
-	{
-		bUseNativeSystrace = true;
-	}
-#endif
-
 #if USE_ANDROID_JNI && !USE_ANDROID_STANDALONE
 	InitializeJavaEventReceivers();
 	AndroidOnBackgroundBinding = FCoreDelegates::ApplicationWillEnterBackgroundDelegate.AddStatic(EnableJavaEventReceivers, false);
@@ -2481,27 +2438,296 @@ bool FAndroidMisc::IsExternalMemoryAndroidHardwareBufferExtensionLoaded()
 	return false;
 }
 
-static UE::FEventCount AvailableConfigRuleVarsEvent;
 static std::atomic_bool bConfigRulesReady = false;
 
 const TMap<FString, FString>& FAndroidMisc::GetConfigRuleVars()
 {
-	if (!bConfigRulesReady)
+	if (!ensure(bConfigRulesReady))
 	{
-		UE::FEventCountToken Token = AvailableConfigRuleVarsEvent.PrepareWait();
-
-		if (!bConfigRulesReady)
-		{
-			FPlatformMisc::LowLevelOutputDebugStringf(TEXT("thread waiting for configrules to be set"));
-			AvailableConfigRuleVarsEvent.Wait(Token);
-			FPlatformMisc::LowLevelOutputDebugStringf(TEXT("done thread waiting for configrules to be set"));
-		}
+		FPlatformMisc::LowLevelOutputDebugStringf(TEXT("Configrule vars requested too early. Using empty data."));
+		// If you encounter this, change the call to use 'static FUtf8StringView blah = GetConfigRulesVariableAsUtf8("blah");'
+		static TMap<FString, FString> EmptyCvars;
+		return EmptyCvars;
 	}
 	
 	// ConfigRuleVars are read-only from here on.
 	return ConfigRulesVariables;
 }
 
+namespace AndroidPlatformMisc
+{
+	// Fills Buffer with the human-readable string for Error via strerror.
+	// Returns "" if Error is 0.
+	static const TCHAR* GetErrnoString(int Error, TCHAR (&Buffer)[256])
+	{
+		if (Error == 0)
+		{
+			Buffer[0] = TEXT('\0');
+		}
+		else
+		{
+			FCString::Snprintf(Buffer, UE_ARRAY_COUNT(Buffer), TEXT("%hs"), strerror(Error));
+		}
+		return Buffer;
+	}
+
+ 	static UE::FSharedMutex EarlyConfigRuleReaderLock;
+
+	class FEarlyConfigRuleReader
+	{
+		static constexpr double SocketReadTimeoutS = 1.5;
+		
+		// \0 denotes abstract unix domain socket name.
+		// This matches java side EarlyConfigrulesSetup in game activity.
+		// java local sockets bind to the abstract domain.
+		static constexpr char SocketName[] = "\0ue.earlyconfigrules";
+		static inline UTF8CHAR EarlyConfRuleBuff[16384] = { 0 };
+		static inline int EarlyConfRuleBuffSize = 0;
+
+		void SetBuffer(const TArrayView<UTF8CHAR> Buff)
+		{
+			// the buffer should only ever be set one time, therefore no memzero required.
+			if(Buff.NumBytes() <= sizeof(EarlyConfRuleBuff))
+			{
+				FMemory::Memcpy(EarlyConfRuleBuff, Buff.GetData(), Buff.NumBytes());
+				EarlyConfRuleBuffSize = Buff.NumBytes();
+			}
+			else
+			{
+				FPlatformMisc::LowLevelOutputDebugStringf(TEXT("FEarlyConfigRuleReader: Failed to set configrules from buffer of size %d, buffer would overflow. (max size %zu)"), Buff.NumBytes(), sizeof(EarlyConfRuleBuff));
+			}
+		}
+
+		void ReadConfigRulesFromSocket()
+		{
+			FPlatformMisc::BeginNamedEvent(FColor(0), TEXT("Start config reader"));
+			ON_SCOPE_EXIT
+			{
+				FPlatformMisc::EndNamedEvent();
+			};
+
+			double t0 = FPlatformTime::Seconds();
+			double tConnected = 0, tVerified = 0, tReceived = 0;
+			int listenfd = socket(AF_UNIX, SOCK_STREAM, 0);
+			TCHAR ErrBuf[256];
+			if (listenfd < 0)
+			{
+				FPlatformMisc::LowLevelOutputDebugStringf(TEXT("FEarlyConfigRuleReader: Failed. socket open failed (%d) %s"), errno, GetErrnoString(errno, ErrBuf));
+				return;
+			}
+
+			ON_SCOPE_EXIT
+			{
+				close(listenfd);
+			};
+
+			struct sockaddr_un addr = { 0 };
+			addr.sun_family = AF_UNIX;
+			// size-1 because abstract socket names are *not* null terminated
+			memcpy(addr.sun_path, SocketName, sizeof(SocketName) - 1);
+
+			int result = -1;
+			while ((FPlatformTime::Seconds() - t0) < SocketReadTimeoutS)
+			{
+				result = connect(listenfd, (struct sockaddr*)&addr, sizeof(addr.sun_family) + sizeof(SocketName) - 1);
+				if (result == 0)
+				{
+					break;
+				}
+
+				if (errno == ECONNREFUSED || errno == ENOENT)
+				{
+					// socket not created yet?
+					FPlatformProcess::SleepNoStats(0.001f);
+				}
+				else
+				{
+					// any other error == failed.
+					FPlatformMisc::LowLevelOutputDebugStringf(TEXT("FEarlyConfigRuleReader: Failed. socket connect failed (%d) %s"), errno, GetErrnoString(errno, ErrBuf));
+					return;
+				}
+			}
+
+			if (result != 0)
+			{
+				FPlatformMisc::LowLevelOutputDebugStringf(TEXT("FEarlyConfigRuleReader: Failed. Timeout when waiting for early rules. %fs)"), (float)(FPlatformTime::Seconds() - t0));
+				return;
+			}
+
+			tConnected = FPlatformTime::Seconds();
+
+			// Abstract Unix sockets are globally visible. Verify the peer UID matches our own
+			// to reject connections from other apps. UID is sufficient on Android as each app
+			// has a unique UID; processes sharing our UID (e.g. app child services) already
+			// have equivalent access to our files and SharedPreferences.
+			// SO_PEERCRED is kernel-populated and cannot be spoofed.
+			struct ucred peercred = {};
+			socklen_t credlen = sizeof(peercred);
+			if (getsockopt(listenfd, SOL_SOCKET, SO_PEERCRED, &peercred, &credlen) != 0)
+			{
+				FPlatformMisc::LowLevelOutputDebugStringf(TEXT("FEarlyConfigRuleReader: Failed to get peer credentials (%d) %s"), errno, GetErrnoString(errno, ErrBuf));
+				return;
+			}
+			if (peercred.uid != static_cast<uid_t>(getuid()))
+			{
+				FPlatformMisc::LowLevelOutputDebugStringf(TEXT("FEarlyConfigRuleReader: Rejecting socket - peer UID %d does not match ours %d"), peercred.uid, getuid());
+				return;
+			}
+
+			tVerified = FPlatformTime::Seconds();
+
+			// set the timeout for recv to use the remaining time, keep a floor of 100ms incase we actually do have something ready.
+			double RemainingTime = FMath::Max(0.1, SocketReadTimeoutS - (FPlatformTime::Seconds() - t0));
+			struct timeval tv = { 0 };
+			tv.tv_sec = (time_t)(RemainingTime);
+			tv.tv_usec = ((time_t)(RemainingTime *1000)-((time_t)(RemainingTime) * 1000)) *1000;
+			if (setsockopt(listenfd, SOL_SOCKET, SO_RCVTIMEO, (char*)&tv, sizeof(tv)))
+			{
+				FPlatformMisc::LowLevelOutputDebugStringf(TEXT("FEarlyConfigRuleReader: Failed to set socket recv timeout (%d) %s"), errno, GetErrnoString(errno, ErrBuf));
+			}
+
+			bool bReading = true;
+			do
+			{
+				int Remaining = sizeof(EarlyConfRuleBuff) - EarlyConfRuleBuffSize;
+				if (Remaining <= 0)
+				{
+					// test for more, do not accept truncated results.
+					char TestFomo;
+					if (recv(listenfd, &TestFomo, sizeof(TestFomo), MSG_DONTWAIT) == 1)
+					{
+						// failed, buffer would overflow.
+						EarlyConfRuleBuffSize = 0;
+						FPlatformMisc::LowLevelOutputDebugStringf(TEXT("FEarlyConfigRuleReader: Failed. Rejecting early configrules data. Buffer would overflow."));
+					}
+					bReading = false;
+					break;
+				}
+				int readamount = recv(listenfd, EarlyConfRuleBuff + EarlyConfRuleBuffSize, Remaining, 0);
+				if (readamount > 0)
+				{
+					EarlyConfRuleBuffSize += readamount;
+				}
+				else if (readamount == 0)
+				{
+					// socket closed. all done.
+					bReading = false;
+				}
+				else if (errno != EINTR)
+				{
+					FPlatformMisc::LowLevelOutputDebugStringf(TEXT("FEarlyConfigRuleReader: Failed to receive (%d) %s"), errno, GetErrnoString(errno, ErrBuf));
+					EarlyConfRuleBuffSize = 0;
+					bReading = false;
+				}
+			} while (bReading);
+
+			tReceived = FPlatformTime::Seconds();
+
+			if (EarlyConfRuleBuffSize > 0)
+			{
+				FPlatformMisc::LowLevelOutputDebugStringf(TEXT("FEarlyConfigRuleReader: Success. %d bytes, connect=%.1fms verify=%.1fms recv=%.1fms total=%.1fms"),
+					EarlyConfRuleBuffSize,
+					(float)(tConnected - t0) * 1000.0f,
+					(float)(tVerified - tConnected) * 1000.0f,
+					(float)(tReceived - tVerified) * 1000.0f,
+					(float)(tReceived - t0) * 1000.0f);
+			}
+			else
+			{
+				FPlatformMisc::LowLevelOutputDebugStringf(TEXT("FEarlyConfigRuleReader: Failed. Did not receive early configrules (%d) %s"), errno, GetErrnoString(errno, ErrBuf));
+			}
+		}
+
+	public:
+
+		FEarlyConfigRuleReader(const TArrayView<UTF8CHAR> Utf8Buff)
+		{
+			// if we have a buffer then avoid waiting for sockets.
+			if (Utf8Buff.IsEmpty())
+			{
+				ReadConfigRulesFromSocket();
+			}
+			else
+			{
+				SetBuffer(Utf8Buff);
+			}
+		}
+
+		// If the socket failed (timeout etc) then an attempt is made to populate during nativeSetConfigRules.
+		// returns true if the buffer was set.
+		// because this is called outside of the static constructor we must use EarlyConfigRuleReaderLock.
+		bool SetBufferIfEmpty(const TArrayView<UTF8CHAR> Buff)
+		{
+ 			UE::TUniqueLock Lock(EarlyConfigRuleReaderLock);
+			if (EarlyConfRuleBuffSize == 0 && Buff.NumBytes() <= sizeof(EarlyConfRuleBuff))
+			{
+				SetBuffer(Buff);
+				return true;
+			}
+			return false;
+		}
+
+		// This is safe as the buffer is forever RO once set.
+		bool GetBuffer(const UTF8CHAR*& BufferOUT, int& BufferSizeOUT) const
+		{
+ 			UE::TSharedLock Lock(EarlyConfigRuleReaderLock);
+			if (EarlyConfRuleBuffSize > 0)
+			{
+				BufferOUT = &EarlyConfRuleBuff[0];
+				BufferSizeOUT = EarlyConfRuleBuffSize;
+				return true;
+			}
+			return false;
+		}
+	};
+
+	// on the first call we will either wait for a socket or use the supplied buffer.
+	// this avoids slower socket path if we get to nativeSetConfigRules first.
+	static FEarlyConfigRuleReader& GetConfigRuleReader(const TArrayView<UTF8CHAR> Utf8Buff = TArrayView<UTF8CHAR>())
+	{
+		static FEarlyConfigRuleReader EarlyConfigRuleReader(Utf8Buff);
+		return EarlyConfigRuleReader;
+	}
+} // namespace AndroidPlatformMisc
+
+// Note: This is a linear search through the buffer. 
+// TODO: cache the stringviews once the allocator is available.
+const FUtf8StringView FAndroidMisc::GetConfigRulesVariableAsUtf8(const FUtf8StringView Key)
+{
+	AndroidPlatformMisc::FEarlyConfigRuleReader& EarlyConfigRuleReader = AndroidPlatformMisc::GetConfigRuleReader();
+
+	UTF8CHAR const* BuffString;
+	int BuffSize;
+	if (EarlyConfigRuleReader.GetBuffer(BuffString, BuffSize))
+	{
+		FUtf8StringView FoundKey;
+		for (int Position = 0; Position<BuffSize;)
+		{
+			uint8 StrSize = BuffString[Position++];
+			// check we have space to read the string.
+			if (StrSize + Position > BuffSize)
+			{
+				break;
+			}
+			FUtf8StringView FoundStr(&BuffString[Position], StrSize);
+			if (!FoundKey.GetData())
+			{
+				FoundKey = FoundStr;
+			}
+			else
+			{
+				if (FoundKey.Equals(Key))
+				{
+					return FoundStr;
+				}
+				FoundKey.Reset();
+			}
+			Position += StrSize;
+		}
+	}
+	return FUtf8StringView();
+}
+
 bool FAndroidMisc::AllowThreadHeartBeat()
 {
 	static uint32 AllowThreadHeartBeatOnce = -1;
@@ -2522,20 +2748,67 @@ bool FAndroidMisc::AllowThreadHeartBeat()
 }
 
 #if USE_ANDROID_JNI
-void JNICALL UE::Jni::FGameActivity::nativeSetConfigRulesVariables(JNIEnv* jenv, jobject thiz, Java::Lang::TArray<jstring>* KeyValuePairs)
+void JNICALL UE::Jni::FGameActivity::nativeSetConfigRulesVariables(JNIEnv* jenv, jobject thiz, Java::Lang::TArray<jstring>* KeyValuePairs, Java::Lang::TArray<jbyte>* Utf8Vars)
 {
-	int32 Count = jenv->GetArrayLength(KeyValuePairs);
+	FPlatformMisc::BeginNamedEvent(FColor(0), TEXT("nativeSetConfigRulesVariables"));
+	ON_SCOPE_EXIT
+	{
+		FPlatformMisc::EndNamedEvent();
+	};
+
+	verify(KeyValuePairs != nullptr);
+	int32 JArrayCount = jenv->GetArrayLength(KeyValuePairs);
+	int32 Count = FMath::DivideAndRoundDown(JArrayCount, 2) * 2; // ensure count is even to match our KV pair assumption later.
+	UE_CLOGF(JArrayCount != Count, LogAndroid, Log, "nativeSetConfigRulesVariables KV pair array must have an even count.");
 	int32 Index = 0;
 	while (Index < Count)
 	{
 		auto javaKey = FJavaHelper::FStringFromLocalRef(jenv, (jstring)(jenv->GetObjectArrayElement(KeyValuePairs, Index++)));
 		auto javaValue = FJavaHelper::FStringFromLocalRef(jenv, (jstring)(jenv->GetObjectArrayElement(KeyValuePairs, Index++)));
-		
+
 		FAndroidMisc::ConfigRulesVariables.Add(javaKey, javaValue);
 	}
 
+	{
+		FPlatformMisc::BeginNamedEvent(FColor(0), TEXT("nativeSetConfigRulesVariables - ECR"));
+		ON_SCOPE_EXIT
+		{
+			FPlatformMisc::EndNamedEvent();
+		};
+
+		jboolean isCopy;
+		jbyte* Utf8bytes = Utf8Vars != nullptr ? jenv->GetByteArrayElements(Utf8Vars, &isCopy) : nullptr;
+		static_assert(sizeof(UTF8CHAR) == sizeof(*Utf8bytes), "Size mismatch in cast");
+		if (Utf8bytes != nullptr)
+		{
+			ON_SCOPE_EXIT
+			{
+				jenv->ReleaseByteArrayElements(Utf8Vars, Utf8bytes, JNI_ABORT);
+			};
+
+			TArrayView Utf8View(reinterpret_cast<UTF8CHAR*>(Utf8bytes), jenv->GetArrayLength(Utf8Vars));
+			AndroidPlatformMisc::FEarlyConfigRuleReader& EarlyConfigRuleReader = AndroidPlatformMisc::GetConfigRuleReader(Utf8View);
+			// if we were beaten to it then we still set the utf8 bytes version (early config data) if we somehow failed to set it.
+			// this would be if an access before nativeSetConfigRulesVariables failed from time out or other socket error.
+			// it should be very unlikely to be required, but is here as best effort to guarantee utf8 configrules variables are around
+			// just as they were with the existing system.
+			if (EarlyConfigRuleReader.SetBufferIfEmpty(Utf8View))
+			{
+				UE_LOGF(LogAndroid, Log, "nativeSetConfigRulesVariables set EarlyConfigRuleReader variables");
+			}
+			else
+			{
+				UE_LOGF(LogAndroid, Log, "nativeSetConfigRulesVariables ignored EarlyConfigRuleReader variables, already set.");
+			}
+		}
+		if (jenv->ExceptionCheck())
+		{
+			UE_LOGF(LogAndroid, Warning, "nativeSetConfigRulesVariables: exception when fetching configrule bytes.");
+			jenv->ExceptionClear();
+		}
+	}
+
 	bConfigRulesReady = true;
-	AvailableConfigRuleVarsEvent.Notify();
 }
 
 #endif
@@ -2558,8 +2831,8 @@ bool FAndroidMisc::UseNewWindowBehavior()
 #if USE_ANDROID_STANDALONE
 	return true;
 #endif
-	static const FString* EnableNewBackgroundBehavior = FAndroidMisc::GetConfigRulesVariable(TEXT("EnableNewBackgroundBehavior"));
-	return EnableNewBackgroundBehavior && EnableNewBackgroundBehavior->ToLower() == TEXT("true");
+	static bool bUseNewWindowBehavior = FPlatformMisc::GetConfigRulesVariableAsUtf8("EnableNewBackgroundBehavior").Equals("true");
+	return bUseNewWindowBehavior;
 }
 
 #if !UE_BUILD_SHIPPING
@@ -2621,24 +2894,13 @@ bool FAndroidMisc::IsDebuggerPresent()
 
 static void WriteTraceMarkerEvent(const ANSICHAR* Text, int32 TraceMarkerFileDescriptor)
 {
-	if (bUseNativeSystrace)
-	{
-		ATrace_beginSection(Text);
-	}
-	else
-	{
-		const int MAX_TRACE_EVENT_LENGTH = 256;
-		ANSICHAR EventBuffer[MAX_TRACE_EVENT_LENGTH];
-		int EventLength = snprintf(EventBuffer, MAX_TRACE_EVENT_LENGTH, "B|%d|%s", getpid(), Text);
-
-		write(TraceMarkerFileDescriptor, EventBuffer, EventLength);
-	}
+	ATrace_beginSection(Text);
 }
 
 void FAndroidMisc::BeginNamedEvent(const struct FColor& Color, const TCHAR* Text)
 {
 	FGenericPlatformMisc::BeginNamedEvent(Color, Text);
-	if (bUseNativeSystrace ? !ATrace_isEnabled() : TraceMarkerFileDescriptor == -1)
+ 	if (!ATrace_isEnabled())
 	{
 		return;
 	}
@@ -2650,7 +2912,7 @@ void FAndroidMisc::BeginNamedEvent(const struct FColor& Color, const TCHAR* Text
 void FAndroidMisc::BeginNamedEvent(const struct FColor& Color, const ANSICHAR* Text)
 {
 	FGenericPlatformMisc::BeginNamedEvent(Color, Text);
-	if (bUseNativeSystrace ? !ATrace_isEnabled() : TraceMarkerFileDescriptor == -1)
+	if (!ATrace_isEnabled() )
 	{
 		return;
 	}
@@ -2661,20 +2923,12 @@ void FAndroidMisc::BeginNamedEvent(const struct FColor& Color, const ANSICHAR* T
 void FAndroidMisc::EndNamedEvent()
 {
 	FGenericPlatformMisc::EndNamedEvent();
-	if (bUseNativeSystrace ? !ATrace_isEnabled() : TraceMarkerFileDescriptor == -1)
+ 	if (!ATrace_isEnabled() )
 	{
 		return;
 	}
 
-	if (bUseNativeSystrace)
-	{
-		ATrace_endSection();
-	}
-	else
-	{
-		const ANSICHAR EventTerminatorChar = 'E';
-		write(TraceMarkerFileDescriptor, &EventTerminatorChar, 1);
-	}
+	ATrace_endSection();
 }
 
 #endif // STATS || ENABLE_STATNAMEDEVENTS
diff --git a/Engine/Source/Runtime/Core/Private/GenericPlatform/GenericPlatformMisc.cpp b/Engine/Source/Runtime/Core/Private/GenericPlatform/GenericPlatformMisc.cpp
--- a/Engine/Source/Runtime/Core/Private/GenericPlatform/GenericPlatformMisc.cpp
+++ b/Engine/Source/Runtime/Core/Private/GenericPlatform/GenericPlatformMisc.cpp
@@ -2443,6 +2443,11 @@ const TMap<FString, FString>& FGenericPlatformMisc::GetConfigRuleVars()
 	return Temp;
 }
 
+const FUtf8StringView FGenericPlatformMisc::GetConfigRulesVariableAsUtf8(const FUtf8StringView Key)
+{
+	return FUtf8StringView();
+}
+
 const FString* FGenericPlatformMisc::GetConfigRulesVariable(const FString& Key) 
 {
 	return FPlatformMisc::GetConfigRuleVars().Find(Key);
diff --git a/Engine/Source/Runtime/Core/Public/Android/AndroidJniGameActivity.h b/Engine/Source/Runtime/Core/Public/Android/AndroidJniGameActivity.h
--- a/Engine/Source/Runtime/Core/Public/Android/AndroidJniGameActivity.h
+++ b/Engine/Source/Runtime/Core/Public/Android/AndroidJniGameActivity.h
@@ -65,7 +65,7 @@ namespace UE::Jni
 		static void JNICALL nativeCrashContextSetDoubleKey(JNIEnv* env, jobject thiz, jstring JavaKey, jdouble JavaValue);
 		static void JNICALL nativeSetObbInfo(JNIEnv* env, jobject thiz, jstring ProjectName, jstring PackageName, jint Version, jint PatchVersion, jstring AppType);
 		static void JNICALL nativeSetAffinityInfo(JNIEnv* env, jobject thiz, jboolean bEnableAffinity, jint bigCoreMask, jint littleCoreMask);
-		static void JNICALL nativeSetConfigRulesVariables(JNIEnv* env, jobject thiz, Java::Lang::TArray<jstring>* KeyValuePairs);
+		static void JNICALL nativeSetConfigRulesVariables(JNIEnv* env, jobject thiz, Java::Lang::TArray<jstring>* KeyValuePairs, Java::Lang::TArray<jbyte>* Utf8Vars);
 		static void JNICALL nativeSetAndroidStartupState(JNIEnv* env, jobject thiz, jboolean bDebuggerAttached);
 		static void JNICALL nativeNetworkChanged(JNIEnv* env, jobject thiz);
 		static void JNICALL nativeResumeMainInit(JNIEnv* env, jobject thiz);
diff --git a/Engine/Source/Runtime/Core/Public/Android/AndroidPlatformMisc.h b/Engine/Source/Runtime/Core/Public/Android/AndroidPlatformMisc.h
--- a/Engine/Source/Runtime/Core/Public/Android/AndroidPlatformMisc.h
+++ b/Engine/Source/Runtime/Core/Public/Android/AndroidPlatformMisc.h
@@ -204,6 +204,9 @@ public:
 	
 	static CORE_API const TMap<FString, FString>& GetConfigRuleVars();
 
+	// return UTF8 versions of the configrules
+	static CORE_API const FUtf8StringView GetConfigRulesVariableAsUtf8(const FUtf8StringView Key);
+
 	/* HasVulkanDriverSupport
 	 * @return true if this Android device supports a Vulkan API Unreal could use
 	 */
diff --git a/Engine/Source/Runtime/Core/Public/GenericPlatform/GenericPlatformMisc.h b/Engine/Source/Runtime/Core/Public/GenericPlatform/GenericPlatformMisc.h
--- a/Engine/Source/Runtime/Core/Public/GenericPlatform/GenericPlatformMisc.h
+++ b/Engine/Source/Runtime/Core/Public/GenericPlatform/GenericPlatformMisc.h
@@ -2040,6 +2040,11 @@ public:
 	*/
 	static CORE_API const FString* GetConfigRulesVariable(const FString& Key);
 
+	/*
+	* Returns the utf8 view of the config rule variable or empty if not found.
+	*/
+	static CORE_API const FUtf8StringView GetConfigRulesVariableAsUtf8(const FUtf8StringView Key);
+
 	/*
 	* Returns the map containing all config rule variables.
 	*/
