FArchiveSerializedPropertyChain::HasUnknownProperty relied on a function pointer set by CoreUObject but it was inlined in an h file so it only got set within CoreUObject's module

#Synced-CL 51689558

diff --git a/Engine/Source/Runtime/Core/Private/Serialization/ArchiveSerializedPropertyChain.cpp b/Engine/Source/Runtime/Core/Private/Serialization/ArchiveSerializedPropertyChain.cpp
new file mode 100644
--- /dev/null
+++ b/Engine/Source/Runtime/Core/Private/Serialization/ArchiveSerializedPropertyChain.cpp
@@ -0,0 +1,12 @@
+// Copyright Epic Games, Inc. All Rights Reserved.
+#include "Serialization/ArchiveSerializedPropertyChain.h"
+
+
+#if WITH_EDITORONLY_DATA
+namespace UE::IDOHelpers
+{
+	PRAGMA_DISABLE_INTERNAL_WARNINGS
+	UE_INTERNAL CORE_API TFunction<bool(TNotNull<FProperty*>)> GIsUnknownPropertyFunc;
+	PRAGMA_ENABLE_INTERNAL_WARNINGS
+}
+#endif
\ No newline at end of file
diff --git a/Engine/Source/Runtime/Core/Public/Serialization/ArchiveSerializedPropertyChain.h b/Engine/Source/Runtime/Core/Public/Serialization/ArchiveSerializedPropertyChain.h
--- a/Engine/Source/Runtime/Core/Public/Serialization/ArchiveSerializedPropertyChain.h
+++ b/Engine/Source/Runtime/Core/Public/Serialization/ArchiveSerializedPropertyChain.h
@@ -2,10 +2,20 @@
 
 #pragma once
 
+#include "CoreTypes.h"
 #include "Containers/Array.h"
 #include "Misc/NotNull.h"
 #include "Templates/Function.h"
 
+#if WITH_EDITORONLY_DATA
+namespace UE::IDOHelpers
+{
+	PRAGMA_DISABLE_INTERNAL_WARNINGS
+	UE_INTERNAL extern CORE_API TFunction<bool(TNotNull<FProperty*>)> GIsUnknownPropertyFunc;
+	PRAGMA_ENABLE_INTERNAL_WARNINGS
+}
+#endif
+
 /** Structure containing the stack of properties that are currently being serialized by an archive */
 struct FArchiveSerializedPropertyChain
 {
@@ -39,7 +49,7 @@ public:
 			++EditorOnlyPropertyStack;
 		}
 		PRAGMA_DISABLE_INTERNAL_WARNINGS
-		if (IsUnknownPropertyFunc.IsSet() && IsUnknownPropertyFunc(InProperty))
+		if (UE::IDOHelpers::GIsUnknownPropertyFunc.IsSet() && UE::IDOHelpers::GIsUnknownPropertyFunc(InProperty))
 		{
 			++UnknownPropertyStack;
 		}
@@ -68,7 +78,7 @@ public:
 			check(EditorOnlyPropertyStack >= 0);
 		}
 		PRAGMA_DISABLE_INTERNAL_WARNINGS
-		if (IsUnknownPropertyFunc.IsSet() && IsUnknownPropertyFunc(InProperty))
+		if (UE::IDOHelpers::GIsUnknownPropertyFunc.IsSet() && UE::IDOHelpers::GIsUnknownPropertyFunc(InProperty))
 		{
 			--UnknownPropertyStack;
 		}
@@ -129,7 +139,7 @@ public:
 	{
 #if WITH_EDITORONLY_DATA
 		PRAGMA_DISABLE_INTERNAL_WARNINGS
-		checkf(IsUnknownPropertyFunc.IsSet(), TEXT("There was no function pointer provided to retrieve this information, so the result will not be correctly returned."))
+		checkf(UE::IDOHelpers::GIsUnknownPropertyFunc.IsSet(), TEXT("There was no function pointer provided to retrieve this information, so the result will not be correctly returned."))
 		PRAGMA_ENABLE_INTERNAL_WARNINGS
 		return UnknownPropertyStack > 0;
 #else // WITH_EDITORONLY_DATA
@@ -160,9 +170,5 @@ private:
 	/** Non-zero if on the current stack of serialized properties is a unknown property */
 	int32 UnknownPropertyStack;
 
-public:
-	PRAGMA_DISABLE_INTERNAL_WARNINGS
-	UE_INTERNAL static inline TFunction<bool(TNotNull<FProperty*>)> IsUnknownPropertyFunc;
-	PRAGMA_ENABLE_INTERNAL_WARNINGS
 #endif // WITH_EDITORONLY_DATA
 };
\ No newline at end of file
diff --git a/Engine/Source/Runtime/CoreUObject/Private/UObject/InstanceDataObjectUtils.cpp b/Engine/Source/Runtime/CoreUObject/Private/UObject/InstanceDataObjectUtils.cpp
--- a/Engine/Source/Runtime/CoreUObject/Private/UObject/InstanceDataObjectUtils.cpp
+++ b/Engine/Source/Runtime/CoreUObject/Private/UObject/InstanceDataObjectUtils.cpp
@@ -34,6 +34,22 @@
 
 DEFINE_LOG_CATEGORY_STATIC(LogInstanceDataObject, Log, VeryVerbose);
 
+namespace UE::IDOHelpers
+{
+	struct FSetupIsUnknownPropertyFnPtr
+	{
+		FSetupIsUnknownPropertyFnPtr()
+		{
+			GIsUnknownPropertyFunc = UE::IsPropertyLoose;
+		}
+		~FSetupIsUnknownPropertyFnPtr()
+		{
+			GIsUnknownPropertyFunc.Reset();
+		}
+	} SetupGetIsUnknownPropertyFnPtr;
+}
+
+
 // Implemented in Class.cpp
 int32 CalculatePropertyIndex(const UStruct* Struct, const FProperty* Property, int32 ArrayIndex);
 
@@ -1249,18 +1265,6 @@ namespace UE
 		return IsInstanceDataObjectSupportEnabled() && Property->GetBoolMetaData(NAME_IsLooseMetadata);
 	}
 
-	struct FSetupIsUnknownPropertyFnPtr
-	{
-		FSetupIsUnknownPropertyFnPtr()
-		{
-			FArchiveSerializedPropertyChain::IsUnknownPropertyFunc = &IsPropertyLoose;
-		}
-		~FSetupIsUnknownPropertyFnPtr()
-		{
-			FArchiveSerializedPropertyChain::IsUnknownPropertyFunc.Reset();
-		}
-	} SetupGetIsUnknownPropertyFnPtr;
-
 	bool IsSerializingUnknownData(const FArchive& Archive, const FProperty* InProperty /* = nullptr*/)
 	{
 		if (!IsInstanceDataObjectSupportEnabled())
