OodleDataCompression : add MinChunkSize argument to CompressParallel so you can opt for higher compression ratio and less parallelism if desired
fix misleading header comments

#Synced-CL 51892244

diff --git a/Engine/Source/Runtime/Core/Private/Compression/OodleDataCompression.cpp b/Engine/Source/Runtime/Core/Private/Compression/OodleDataCompression.cpp
--- a/Engine/Source/Runtime/Core/Private/Compression/OodleDataCompression.cpp
+++ b/Engine/Source/Runtime/Core/Private/Compression/OodleDataCompression.cpp
@@ -613,12 +613,20 @@ int64 CompressParallelSub(
 							TArray64< TArray64<uint8> > & OutChunkCompressedData,
 							const void * InUncompressedData, int64 UncompressedSize,
 							ECompressor Compressor,	ECompressionLevel Level,
-							bool CompressIndependentChunks)
+							bool CompressIndependentChunks,
+							int64 MinChunkSize)
 {
 	// beware: doing the split based on NumWorkers means that compressed output depends on the machine (if ! CompressIndependentChunks)
 	//int32 NumWorkers = FTaskGraphInterface::Get().GetNumWorkerThreads();
 	int32 NumWorkers = 128; // hard coded so that output is the same on all machines
 	int64 ChunkLen = OodleLZ_MakeSeekChunkLen(UncompressedSize,NumWorkers);
+
+	// if CompressIndependentChunks is set we always do OODLELZ_BLOCK_LEN cuts and MinChunkSize is ignored
+	if ( !CompressIndependentChunks && ChunkLen < MinChunkSize )
+	{
+		constexpr uint64 BlockLenMask = OODLELZ_BLOCK_LEN-1;
+		ChunkLen = (MinChunkSize + BlockLenMask) & (~BlockLenMask);
+	}
 	int64 NumChunks = (UncompressedSize + ChunkLen-1)/ChunkLen;
 
 	OutChunkCompressedData.SetNum( NumChunks );
@@ -673,16 +681,17 @@ int64 CompressParallel(
 							void * OutCompressedData, int64 CompressedBufferSize,
 							const void * InUncompressedData, int64 UncompressedSize,
 							ECompressor Compressor,	ECompressionLevel Level,
-							bool CompressIndependentChunks)
+							bool CompressIndependentChunks,
+							int64 MinChunkSize)
 {
-	if ( UncompressedSize <= OODLELZ_BLOCK_LEN )
+	if ( UncompressedSize <= FMath::Max(OODLELZ_BLOCK_LEN,MinChunkSize) )
 	{
 		return Compress(OutCompressedData,CompressedBufferSize,
 			InUncompressedData,UncompressedSize,Compressor,Level,CompressIndependentChunks);
 	}
 
 	TArray64< TArray64<uint8> > Chunks;
-	int64 TotalCompLen = CompressParallelSub(Chunks,InUncompressedData,UncompressedSize,Compressor,Level,CompressIndependentChunks);
+	int64 TotalCompLen = CompressParallelSub(Chunks,InUncompressedData,UncompressedSize,Compressor,Level,CompressIndependentChunks,MinChunkSize);
 	if ( TotalCompLen <= 0 )
 	{
 		UE_LOG(OodleDataCompression,Error,TEXT("CompressParallelSub failed\n"));
@@ -723,12 +732,13 @@ int64 CompressParallel(
 							TArray64<uint8> & OutCompressedArray,
 							const void * InUncompressedData, int64 UncompressedSize,
 							ECompressor Compressor,	ECompressionLevel Level,
-							bool CompressIndependentChunks)
+							bool CompressIndependentChunks,
+							int64 MinChunkSize)
 {
 	// appends to OutCompressedArray
 	int64 StartNum = OutCompressedArray.Num();
 
-	if ( UncompressedSize <= OODLELZ_BLOCK_LEN )
+	if ( UncompressedSize <= FMath::Max(OODLELZ_BLOCK_LEN,MinChunkSize) )
 	{
 		int64 Reserve = CompressedBufferSizeNeeded(UncompressedSize);
 		OutCompressedArray.SetNum( StartNum + Reserve );
@@ -747,7 +757,7 @@ int64 CompressParallel(
 	}
 
 	TArray64< TArray64<uint8> > Chunks;
-	int64 TotalCompLen = CompressParallelSub(Chunks,InUncompressedData,UncompressedSize,Compressor,Level,CompressIndependentChunks);
+	int64 TotalCompLen = CompressParallelSub(Chunks,InUncompressedData,UncompressedSize,Compressor,Level,CompressIndependentChunks,MinChunkSize);
 	if ( TotalCompLen <= 0 )
 	{
 		UE_LOG(OodleDataCompression,Error,TEXT("CompressParallelSub failed\n"));
@@ -788,6 +798,7 @@ bool DecompressParallel(
 	//	(eg. say if you did seek resets at OODLELZ_BLOCK_LEN*4)
 	// then we would have to change this logic to instead scan the comp data
 	//	and find the seek resets, using OodleLZ_GetCompressedStepForRawStep
+	// -> we could improve this if there's a use case for it
 
 	int32 NumWorkers = FTaskGraphInterface::Get().GetNumWorkerThreads();
 	int32 ChunkLen = OodleLZ_MakeSeekChunkLen(UncompressedSize,NumWorkers*3);
diff --git a/Engine/Source/Runtime/Core/Public/Compression/OodleDataCompression.h b/Engine/Source/Runtime/Core/Public/Compression/OodleDataCompression.h
--- a/Engine/Source/Runtime/Core/Public/Compression/OodleDataCompression.h
+++ b/Engine/Source/Runtime/Core/Public/Compression/OodleDataCompression.h
@@ -190,9 +190,10 @@ int64 CORE_API Compress(
 * CompressedBufferSize must be >= CompressedBufferSizeNeeded(UncompressedSize)
 * CompressParallel can be used in all cases instead of Compress
 * it is fast even on small buffers (will just run synchronously on the calling thread)
-* the output compressed data is the same as Compress
+* The output made by this function is a single compressed stream of the same format as Compress()
 * can be decoded with either Decompress or DecompressParallel
 * DecompressParallel can only parallelize if CompressIndependentChunks is true at encode time
+* NOTE: CompressIndependentChunks and small MinChunkSize do hurt compression ratio (compared to non-parallel Compress)
 *
 * @param	OutCompressedData			output buffer where compressed data is written
 * @param	CompressedBufferSize		bytes available to write in OutCompressedData
@@ -201,22 +202,30 @@ int64 CORE_API Compress(
 * @param	Compressor					ECompressor to encode with (this is saved in the stream)
 * @param	Level						ECompressionLevel to encode with (this is not saved in the stream)
 * @param	CompressIndependentChunks	(optional) should chunks be made independent (allows for parallel decode)
+* @param	MinChunkSize				minimum chunk size; making this larger (eg. 16 MB) increases compression ratio but hurts parallelism; only used if CompressIndependentChunks is false
 * @return								Compressed size written or zero for failure
+*
+* Calling Compress() or CompressParallel() does not affect whether or not the data can be decoded parallel.
+* That is controlled by CompressIndependentChunks.
+* The choice of Compress vs CompressParallel affects encode speed and compression ratio.
+* CompressParallel lets you sacrifice some compression ratio for faster encodes.
 */
 int64 CORE_API CompressParallel(
 							void * OutCompressedData, int64 CompressedBufferSize,
 							const void * InUncompressedData, int64 UncompressedSize,
 							ECompressor Compressor,	ECompressionLevel Level,
-							bool CompressIndependentChunks = false);
+							bool CompressIndependentChunks = false,
+							int64 MinChunkSize = 0);
 							
 /**
 * Encode provided data with chosen Compressor and Level, using multiple threads for large buffers
 * CompressedBufferSize must be >= CompressedBufferSizeNeeded(UncompressedSize)
 * CompressParallel can be used in all cases instead of Compress
 * it is fast even on small buffers (will just run synchronously on the calling thread)
-* the output compressed data is the same as Compress
+* The output made by this function is a single compressed stream of the same format as Compress()
 * can be decoded with either Decompress or DecompressParallel
 * DecompressParallel can only parallelize if CompressIndependentChunks is true at encode time
+* NOTE: CompressIndependentChunks and small MinChunkSize do hurt compression ratio (compared to non-parallel Compress)
 *
 * @param	OutCompressedData			output buffer where compressed data is written; array is appended to
 * @param	InUncompressedData			input buffer containing data to compress
@@ -224,13 +233,20 @@ int64 CORE_API CompressParallel(
 * @param	Compressor					ECompressor to encode with (this is saved in the stream)
 * @param	Level						ECompressionLevel to encode with (this is not saved in the stream)
 * @param	CompressIndependentChunks	should chunks be made independent (allows for parallel decode)
+* @param	MinChunkSize				minimum chunk size; making this larger (eg. 16 MB) increases compression ratio but hurts parallelism; only used if CompressIndependentChunks is false
 * @return								Compressed size written or zero for failure
+*
+* Calling Compress() or CompressParallel() does not affect whether or not the data can be decoded parallel.
+* That is controlled by CompressIndependentChunks.
+* The choice of Compress vs CompressParallel affects encode speed and compression ratio.
+* CompressParallel lets you sacrifice some compression ratio for faster encodes.
 */
 int64 CORE_API CompressParallel(
 							TArray64<uint8> & OutCompressedData,
 							const void * InUncompressedData, int64 UncompressedSize,
 							ECompressor Compressor,	ECompressionLevel Level,
-							bool CompressIndependentChunks = false);
+							bool CompressIndependentChunks = false,
+							int64 MinChunkSize = 0);
 
 /**
 * Decode compressed data that was made by Compress
