[InheritedContextTest]
- Fix random failure because we didn't wait on task to finish correctly

#Synced-CL 51868793

diff --git a/Engine/Source/Runtime/Core/Tests/Async/InheritedContextTest.cpp b/Engine/Source/Runtime/Core/Tests/Async/InheritedContextTest.cpp
--- a/Engine/Source/Runtime/Core/Tests/Async/InheritedContextTest.cpp
+++ b/Engine/Source/Runtime/Core/Tests/Async/InheritedContextTest.cpp
@@ -58,12 +58,12 @@ namespace InheritedContextTests
 		// Using an event prevents retraction since we want those tasks to run on worker threads
 		// to validate propagation.
 		FTaskEvent EventA(TEXT("EventA"));
-		Launch(TEXT("TaskA"), [&EventA]()
+		auto TaskA = Launch(TEXT("TaskA"), [&EventA]()
 		{
 			CHECK(GTestFlag);
 
 			FTaskEvent EventB(TEXT("EventB"));
-			Launch(TEXT("TaskB"), [&EventB]()
+			auto TaskB = Launch(TEXT("TaskB"), [&EventB]()
 			{
 				CHECK(GTestFlag);
 
@@ -72,10 +72,17 @@ namespace InheritedContextTests
 
 			EventB.Wait();
 
+			// Wait until the task is complete even after the event is triggered
+			// otherwise the inherited scope might revert any value after we're out of scope
+			TaskB.Wait();
 			EventA.Trigger();
 		});
 
 		EventA.Wait();
+
+		// Wait until the task is complete even after the event is triggered
+		// otherwise the inherited scope might revert any value after we're out of scope
+		TaskA.Wait();
 	}
 
 	// Test TInheritedGuardValue propagation (combines guard + extension scope)
@@ -91,12 +98,12 @@ namespace InheritedContextTests
 		// Using an event prevents retraction since we want those tasks to run on worker threads
 		// to validate propagation.
 		FTaskEvent EventA(TEXT("EventA"));
-		Launch(TEXT("TaskA"), [&EventA]()
+		auto TaskA = Launch(TEXT("TaskA"), [&EventA]()
 		{
 			CHECK(GTestFlag);
 
 			FTaskEvent EventB(TEXT("EventB"));
-			Launch(TEXT("TaskB"), [&EventB]()
+			auto TaskB = Launch(TEXT("TaskB"), [&EventB]()
 			{
 				CHECK(GTestFlag);
 
@@ -105,10 +112,17 @@ namespace InheritedContextTests
 
 			EventB.Wait();
 
+			// Wait until the task is complete even after the event is triggered
+			// otherwise the inherited scope might revert any value after we're out of scope
+			TaskB.Wait();
 			EventA.Trigger();
 		});
 
 		EventA.Wait();
+
+		// Wait until the task is complete even after the event is triggered
+		// otherwise the inherited scope might revert any value after we're out of scope
+		TaskA.Wait();
 	}
 
 	// Test TInheritedGuardValue with pre-built extension (avoids per-guard allocation)
@@ -122,7 +136,7 @@ namespace InheritedContextTests
 		TInheritedGuardValue<bool> Guard(GTestFlagExt, GTestFlag, true);
 
 		FTaskEvent EventA(TEXT("EventA"));
-		Launch(TEXT("TaskA"), [&EventA]()
+		auto TaskA = Launch(TEXT("TaskA"), [&EventA]()
 		{
 			CHECK(GTestFlag);
 
@@ -130,6 +144,10 @@ namespace InheritedContextTests
 		});
 
 		EventA.Wait();
+
+		// Wait until the task is complete even after the event is triggered
+		// otherwise the inherited scope might revert any value after we're out of scope
+		TaskA.Wait();
 	}
 
 	// Test getter/setter variant of MakeInheritedContextExtension
@@ -143,23 +161,31 @@ namespace InheritedContextTests
 		UE::FInheritedContextExtensionScope ActivateFlag(GAccessorFlagExt);
 
 		FTaskEvent EventA(TEXT("EventA"));
-		Launch(TEXT("TaskA"), [&EventA]()
+		auto TaskA = Launch(TEXT("TaskA"), [&EventA]()
 		{
 			CHECK(GetAccessorFlag());
 
 			FTaskEvent EventB(TEXT("EventB"));
-			Launch(TEXT("TaskB"), [&EventB]()
+			auto TaskB = Launch(TEXT("TaskB"), [&EventB]()
 			{
 				CHECK(GetAccessorFlag());
 				EventB.Trigger();
 			});
 
 			EventB.Wait();
+
+			// Wait until the task is complete even after the event is triggered
+			// otherwise the inherited scope might revert any value after we're out of scope
+			TaskB.Wait();
 			EventA.Trigger();
 		});
 
 		EventA.Wait();
 
+		// Wait until the task is complete even after the event is triggered
+		// otherwise the inherited scope might revert any value after we're out of scope
+		TaskA.Wait();
+
 		SetAccessorFlag(false);
 	}
 
@@ -171,7 +197,7 @@ namespace InheritedContextTests
 		TInheritedGuardValueAccessors<&GetAccessorFlag, &SetAccessorFlag> Guard(true);
 
 		FTaskEvent EventA(TEXT("EventA"));
-		Launch(TEXT("TaskA"), [&EventA]()
+		auto TaskA = Launch(TEXT("TaskA"), [&EventA]()
 		{
 			CHECK(GetAccessorFlag());
 
@@ -179,6 +205,10 @@ namespace InheritedContextTests
 		});
 
 		EventA.Wait();
+
+		// Wait until the task is complete even after the event is triggered
+		// otherwise the inherited scope might revert any value after we're out of scope
+		TaskA.Wait();
 	}
 
 	// Test propagation through ParallelFor
@@ -219,26 +249,32 @@ namespace InheritedContextTests
 		// Using an event prevents retraction since we want those tasks to run on worker threads
 		// to validate propagation.
 		FTaskEvent EventA(TEXT("EventA"));
-		Launch(TEXT("TaskA"), [&EventA]()
+		auto TaskA = Launch(TEXT("TaskA"), [&EventA]()
 		{
 			FPlatformProcess::Sleep(0.1f);
 
 			FTaskEvent EventB(TEXT("EventB"));
+			auto TaskB = Launch(TEXT("TaskB"), [&EventB]()
 			{
-				Launch(TEXT("TaskB"), [&EventB]()
-				{
-					FPlatformProcess::Sleep(0.1f);
-					EventB.Trigger();
-				});
-			}
+				FPlatformProcess::Sleep(0.1f);
+				EventB.Trigger();
+			});
 
 			EventB.Wait(); // Additional 0.1s for waiting on the other task
 
+			// Wait until the task is complete even after the event is triggered
+			// otherwise the inherited scope might revert any value after we're out of scope
+			TaskB.Wait();
+
 			EventA.Trigger();
 		});
 
 		EventA.Wait();
 
+		// Wait until the task is complete even after the event is triggered
+		// otherwise the inherited scope might revert any value after we're out of scope
+		TaskA.Wait();
+
 		FTaskTime& TaskTime = TaskTimingExtension.GetState<FTaskTime>();
 
 		uint64 TotalCycles = TaskTime.GameThreadTime.load() + TaskTime.WorkerThreadTime.load();
@@ -274,7 +310,7 @@ namespace InheritedContextTests
 		UE::FInheritedContextExtensionScope Scope(Ext);
 
 		FTaskEvent EventA(TEXT("EventA"));
-		Launch(TEXT("TaskA"), [&Ext, &EventA]() 
+		auto TaskA = Launch(TEXT("TaskA"), [&Ext, &EventA]() 
 		{
 			// TaskA's scope applied once
 			CHECK(GApplyCount.load(std::memory_order_relaxed) == 1);
@@ -285,7 +321,7 @@ namespace InheritedContextTests
 			UE::FInheritedContextExtensionScope InnerScope(Ext);
 
 			FTaskEvent EventB(TEXT("EventB"));
-			Launch(TEXT("TaskB"), [&EventB]() 
+			auto TaskB = Launch(TEXT("TaskB"), [&EventB]() 
 			{
 				// With correct dedup: TaskA(1) + TaskB(1) = 2
 				// Without dedup: TaskA(1) + TaskB(2) = 3
@@ -294,10 +330,17 @@ namespace InheritedContextTests
 			});
 
 			EventB.Wait();
+
+			// Wait until the task is complete even after the event is triggered
+			// otherwise the inherited scope might revert any value after we're out of scope
+			TaskB.Wait();
 			EventA.Trigger();
 		});
 
 		EventA.Wait();
+		// Wait until the task is complete even after the event is triggered
+		// otherwise the inherited scope might revert any value after we're out of scope
+		TaskA.Wait();
 
 		// All scopes restored
 		CHECK(GApplyCount.load(std::memory_order_relaxed) == 0);
@@ -317,12 +360,12 @@ namespace InheritedContextTests
 
 		// No FInheritedContextExtensionScope - flag should not propagate to worker threads
 		FTaskEvent EventA(TEXT("EventA"));
-		Launch(TEXT("TaskA"), [&EventA]() 
+		auto TaskA = Launch(TEXT("TaskA"), [&EventA]() 
 		{
 			CHECK(!GTestFlag);
 
 			FTaskEvent EventB(TEXT("EventB"));
-			Launch(TEXT("TaskB"), [&EventB]() 
+			auto TaskB = Launch(TEXT("TaskB"), [&EventB]() 
 			{
 				CHECK(!GTestFlag);
 
@@ -330,11 +373,17 @@ namespace InheritedContextTests
 			});
 
 			EventB.Wait();
+			// Wait until the task is complete even after the event is triggered
+			// otherwise the inherited scope might revert any value after we're out of scope
+			TaskB.Wait();
 
 			EventA.Trigger();
 		});
 
 		EventA.Wait();
+		// Wait until the task is complete even after the event is triggered
+		// otherwise the inherited scope might revert any value after we're out of scope
+		TaskA.Wait();
 	}
 
 	TEST_CASE_NAMED(FInheritedContextMultipleExtensions, "System::Core::Async::InheritedContext::MultipleExtensions", "[ApplicationContextMask][EngineFilter]")
@@ -351,7 +400,7 @@ namespace InheritedContextTests
 		TInheritedGuardValue<bool> Guard4([]() -> bool& { return GTestFlags[4]; }, true);
 
 		FTaskEvent EventA(TEXT("EventA"));
-		Launch(TEXT("TaskA"), [&EventA]()
+		auto TaskA = Launch(TEXT("TaskA"), [&EventA]()
 		{
 			CHECK(GTestFlags[0]);
 			CHECK(GTestFlags[1]);
@@ -360,7 +409,7 @@ namespace InheritedContextTests
 			CHECK(GTestFlags[4]);
 
 			FTaskEvent EventB(TEXT("EventB"));
-			Launch(TEXT("TaskB"), [&EventB]()
+			auto TaskB = Launch(TEXT("TaskB"), [&EventB]()
 			{
 				CHECK(GTestFlags[0]);
 				CHECK(GTestFlags[1]);
@@ -373,10 +422,17 @@ namespace InheritedContextTests
 
 			EventB.Wait();
 
+			// Wait until the task is complete even after the event is triggered
+			// otherwise the inherited scope might revert any value after we're out of scope
+			TaskB.Wait();
 			EventA.Trigger();
 		});
 
 		EventA.Wait();
+
+		// Wait until the task is complete even after the event is triggered
+		// otherwise the inherited scope might revert any value after we're out of scope
+		TaskA.Wait();
 	}
 
 	TEST_CASE_NAMED(FInheritedContextExtensionGoesOutOfScopeBeforeExecution, "System::Core::Async::InheritedContext::ExtensionGoesOutOfScopeBeforeExecution", "[ApplicationContextMask][EngineFilter]")
@@ -388,13 +444,14 @@ namespace InheritedContextTests
 		FTaskEvent EventA(TEXT("EventA"));
 		FTaskEvent StartEvent(TEXT("StartEvent"));
 
+		FTask TaskA;
 		{
 			UE::FInheritedContextExtension GTestFlagExt = UE::MakeInheritedContextExtension([]() -> bool& { return GTestFlag; });
 		
 			TGuardValue<bool> Guard(GTestFlag, true);
 			UE::FInheritedContextExtensionScope ActivateFlag(GTestFlagExt);
 
-			Launch(TEXT("TaskA"), [&EventA]()
+			TaskA = Launch(TEXT("TaskA"), [&EventA]()
 			{
 				CHECK(GTestFlag);
 				// test toggling the flag for the next task while we're at it
@@ -418,6 +475,10 @@ namespace InheritedContextTests
 		CHECK(!GTestFlag);
 		StartEvent.Trigger();
 		EventA.Wait();
+
+		// Wait until the task is complete even after the event is triggered
+		// otherwise the inherited scope might revert any value after we're out of scope
+		TaskA.Wait();
 	}
 }
 
