[Traces] Add -tracefiletimestamps option for trace file naming. It will append the date and time to the -tracefile name provided.

#Synced-CL 51871254

diff --git a/Engine/Source/Runtime/Core/Private/ProfilingDebugging/TraceAuxiliary.cpp b/Engine/Source/Runtime/Core/Private/ProfilingDebugging/TraceAuxiliary.cpp
--- a/Engine/Source/Runtime/Core/Private/ProfilingDebugging/TraceAuxiliary.cpp
+++ b/Engine/Source/Runtime/Core/Private/ProfilingDebugging/TraceAuxiliary.cpp
@@ -254,6 +254,7 @@ static FDelegateHandle GOnPostForkHandle;
 // Whether to start tracing automatically at start or wait to initiate via Console Command.
 // This value can also be set by passing '-traceautostart=[0|1]' on command line.
 static bool GTraceAutoStart = UE_TRACE_AUTOSTART ? true : false;
+static bool GTraceFileTimestamps = false;
 
 ////////////////////////////////////////////////////////////////////////////////
 bool FTraceAuxiliaryImpl::IsParentProcessAndPreFork()
@@ -1000,6 +1001,14 @@ bool FTraceAuxiliaryImpl::FinalizeFilePath(const TCHAR* InPath, FString& OutPath
 		WritePath = FPaths::SetExtension(WritePath, TEXT(".utrace"));
 	}
 
+	// Append timestamp to filename if requested via -tracefiletimestamps
+	if (GTraceFileTimestamps)
+	{
+		const FString BasePath = FPaths::GetBaseFilename(WritePath, false); // false = include path
+		const FDateTime Now = FDateTime::Now();
+		WritePath = FString::Printf(TEXT("%s_%s.utrace"), *BasePath, *Now.ToString(TEXT("%Y%m%d_%H%M%S")));
+	}
+
 	// Finally make sure the path is platform friendly
 	IFileManager& FileManager = IFileManager::Get();
 	FString NativePath = FileManager.ConvertToAbsolutePathForExternalAppForWrite(*WritePath);
@@ -1895,6 +1904,7 @@ static bool StartFromCommandlineArguments(const TCHAR* CommandLine)
 	// Setup options
 	FTraceAuxiliary::FOptions Opts;
 	Opts.bTruncateFile = FParse::Param(CommandLine, TEXT("tracefiletrunc"));
+	GTraceFileTimestamps = FParse::Param(CommandLine, TEXT("tracefiletimestamps"));
 
 	// Find if a connection type is specified
 	FString Parameter;
