[InheritedContext]
- Fix apply order to be the same as the scope order
- Add tests to validate

#Synced-CL 51934851

diff --git a/Engine/Source/Runtime/Core/Private/Async/InheritedContext.cpp b/Engine/Source/Runtime/Core/Private/Async/InheritedContext.cpp
--- a/Engine/Source/Runtime/Core/Private/Async/InheritedContext.cpp
+++ b/Engine/Source/Runtime/Core/Private/Async/InheritedContext.cpp
@@ -192,9 +192,12 @@ FInheritedContextScope::FInheritedContextScope(TSharedPtr<const FAppTime> InAppT
 		// Allocate saved state block with the same layout
 		ExtensionSavedState = FInheritedContextExtensionData::AllocateLike(InExtensionData);
 
-		// Apply each extension: write captured value to thread state, save previous state
+		// Apply each extension in reverse order: entries are stored in LIFO order
+		// (most recent scope first), so iterating backward replays them in definition
+		// order — the innermost (last defined) scope is applied last and its value wins,
+		// matching what the parent thread sees.
 		const FInheritedContextExtensionData::FEntry* Entries = InExtensionData->GetEntries();
-		for (int32 i = 0; i < InExtensionData->NumEntries; ++i)
+		for (int32 i = InExtensionData->NumEntries - 1; i >= 0; --i)
 		{
 			Entries[i].Ext.Impl->Apply(
 				InExtensionData->GetData(i),
@@ -218,11 +221,13 @@ FInheritedContextScope::FInheritedContextScope(TSharedPtr<const FAppTime> InAppT
 
 FInheritedContextScope::~FInheritedContextScope()
 {
-	// Restore extensions in reverse order
+	// Restore extensions in forward order: since entries are in LIFO order and
+	// Apply iterated backward (definition order), restoring forward unwinds in
+	// reverse definition order — last applied is first restored (standard RAII).
 	if (ExtensionSavedState)
 	{
 		const FInheritedContextExtensionData::FEntry* Entries = ExtensionSavedState->GetEntries();
-		for (int32 i = ExtensionSavedState->NumEntries - 1; i >= 0; --i)
+		for (int32 i = 0; i < ExtensionSavedState->NumEntries; ++i)
 		{
 			Entries[i].Ext.Impl->Restore(
 				ExtensionSavedState->GetData(i),
@@ -278,7 +283,7 @@ void FInheritedContextBase::CaptureInheritedContext()
 	{
 		TArray<const FInheritedContextExtension*, TInlineAllocator<8>> Extensions;
 
-		// Collect from active scope chain (most recent first)
+		// Collect from active scope chain (most recent first due to LIFO linked list)
 		for (FInheritedContextExtensionScope* Scope = GExtensionScopeHead; Scope; Scope = Scope->Next)
 		{
 			Extensions.Add(&Scope->Extension);
diff --git a/Engine/Source/Runtime/Core/Tests/Async/InheritedContextTest.cpp b/Engine/Source/Runtime/Core/Tests/Async/InheritedContextTest.cpp
--- a/Engine/Source/Runtime/Core/Tests/Async/InheritedContextTest.cpp
+++ b/Engine/Source/Runtime/Core/Tests/Async/InheritedContextTest.cpp
@@ -435,6 +435,134 @@ namespace InheritedContextTests
 		TaskA.Wait();
 	}
 
+	// Helpers for ordering test — tracks the sequence in which Apply and Restore are
+	// called across multiple extensions. Uses compile-time IDs baked into template
+	// instantiations so each extension records a unique identifier.
+	struct FOrderState {};
+	static std::atomic<int32> GOrderApplyCounter{0};
+	static std::atomic<int32> GOrderRestoreCounter{0};
+	static int32 GApplySequence[3];
+	static int32 GRestoreSequence[3];
+
+	template<int32 Id>
+	static int32 OnApplyOrdered()
+	{
+		int32 Pos = GOrderApplyCounter.fetch_add(1, std::memory_order_seq_cst);
+		GApplySequence[Pos] = Id;
+		return 0;
+	}
+
+	template<int32 Id>
+	static void OnRestoreOrdered(const int32& /*Saved*/, FOrderState& /*State*/)
+	{
+		int32 Pos = GOrderRestoreCounter.fetch_add(1, std::memory_order_seq_cst);
+		GRestoreSequence[Pos] = Id;
+	}
+
+	// Verify that when multiple extensions are captured and applied on a worker thread,
+	// Apply is called in definition order (matching the order scopes were entered on
+	// the parent thread) and Restore is called in reverse definition order (standard
+	// RAII unwinding). This ensures the innermost scope's value wins on the worker
+	// thread, consistent with what the parent thread sees.
+	TEST_CASE_NAMED(FInheritedContextMultipleExtensionsApplyOrder, "System::Core::Async::InheritedContext::MultipleExtensionsApplyOrder", "[ApplicationContextMask][EngineFilter]")
+	{
+		using namespace UE::Tasks;
+
+		// Reset global tracking state
+		GOrderApplyCounter.store(0, std::memory_order_seq_cst);
+		GOrderRestoreCounter.store(0, std::memory_order_seq_cst);
+		for (int32 i = 0; i < 3; ++i)
+		{
+			GApplySequence[i] = -1;
+			GRestoreSequence[i] = -1;
+		}
+
+		// Create 3 stateful extensions that track their Apply/Restore order.
+		// Each has a unique compile-time ID (0, 1, 2).
+		UE::FInheritedContextExtension Ext0 = UE::MakeStatefulInheritedContextExtension<&OnApplyOrdered<0>, &OnRestoreOrdered<0>, FOrderState>();
+		UE::FInheritedContextExtension Ext1 = UE::MakeStatefulInheritedContextExtension<&OnApplyOrdered<1>, &OnRestoreOrdered<1>, FOrderState>();
+		UE::FInheritedContextExtension Ext2 = UE::MakeStatefulInheritedContextExtension<&OnApplyOrdered<2>, &OnRestoreOrdered<2>, FOrderState>();
+
+		// Activate scopes in definition order 0, 1, 2.
+		// The scope chain linked list will be 2 -> 1 -> 0 (LIFO), and entries are
+		// captured in that order. Apply iterates backward to replay in definition order.
+		UE::FInheritedContextExtensionScope Scope0(Ext0);
+		UE::FInheritedContextExtensionScope Scope1(Ext1);
+		UE::FInheritedContextExtensionScope Scope2(Ext2);
+
+		// Using an event prevents retraction since we want the task to run on a worker
+		// thread to validate cross-thread ordering.
+		FTaskEvent EventA(TEXT("EventA"));
+		auto TaskA = Launch(TEXT("TaskA"), [&EventA]()
+		{
+			// By this point, RestoreInheritedContext has been called by the task system,
+			// which applied all 3 extensions via the FInheritedContextScope constructor.
+			// Verify that Apply was called in definition order: 0, 1, 2.
+			// This ensures the innermost scope (2) is applied last and its value wins.
+			CHECK(GOrderApplyCounter.load(std::memory_order_seq_cst) == 3);
+			CHECK(GApplySequence[0] == 0);
+			CHECK(GApplySequence[1] == 1);
+			CHECK(GApplySequence[2] == 2);
+
+			EventA.Trigger();
+		});
+
+		EventA.Wait();
+
+		// Wait until the task is complete even after the event is triggered
+		// otherwise the inherited scope might revert any value after we're out of scope
+		TaskA.Wait();
+
+		// After the task completes and its FInheritedContextScope destructor runs,
+		// Restore should have been called in reverse definition order: 2, 1, 0
+		// (standard RAII unwinding — last applied is first restored)
+		CHECK(GOrderRestoreCounter.load(std::memory_order_seq_cst) == 3);
+		CHECK(GRestoreSequence[0] == 2);
+		CHECK(GRestoreSequence[1] == 1);
+		CHECK(GRestoreSequence[2] == 0);
+	}
+
+	// Verify that when multiple nested scopes modify the same variable via the same
+	// pre-built extension, the worker thread sees the innermost (last set) value.
+	// All scopes share the same Impl pointer, so CaptureInheritedContext captures
+	// the variable's current value (set by the innermost guard) for each entry.
+	TEST_CASE_NAMED(FInheritedContextNestedScopesSameVariable, "System::Core::Async::InheritedContext::NestedScopesSameVariable", "[ApplicationContextMask][EngineFilter]")
+	{
+		using namespace UE::Tasks;
+
+		static thread_local int32 GValue = 0;
+		static UE::FInheritedContextExtension GValueExt = UE::MakeInheritedContextExtension([]() -> int32& { return GValue; });
+
+		// Three nested guard+scope pairs on the same variable and extension.
+		// Each guard sets a different value; the innermost (30) should win.
+		TGuardValue<int32> Guard1(GValue, 10);
+		UE::FInheritedContextExtensionScope Scope1(GValueExt);
+
+		TGuardValue<int32> Guard2(GValue, 20);
+		UE::FInheritedContextExtensionScope Scope2(GValueExt);
+
+		TGuardValue<int32> Guard3(GValue, 30);
+		UE::FInheritedContextExtensionScope Scope3(GValueExt);
+
+		// Using an event prevents retraction since we want the task to run on a worker
+		// thread to validate propagation.
+		FTaskEvent EventA(TEXT("EventA"));
+		auto TaskA = Launch(TEXT("TaskA"), [&EventA]()
+		{
+			// The innermost scope's value (30) should be visible on the worker thread,
+			// consistent with what the parent thread sees.
+			CHECK(GValue == 30);
+
+			EventA.Trigger();
+		});
+
+		EventA.Wait();
+
+		// Wait until the task is complete even after the event is triggered
+		// otherwise the inherited scope might revert any value after we're out of scope
+		TaskA.Wait();
+	}
+
 	TEST_CASE_NAMED(FInheritedContextExtensionGoesOutOfScopeBeforeExecution, "System::Core::Async::InheritedContext::ExtensionGoesOutOfScopeBeforeExecution", "[ApplicationContextMask][EngineFilter]")
 	{
 		using namespace UE::Tasks;
