Integrating shelve 51926103:
Added hack to force non-scoped enums in dynamic delegate signatures to be a byte.  This should be caught in UHT, and the hack can be reversed once it is rejected fully.

#tests Cooked Client- verified the dynamic delegate was passing the correct values for the HUD visibility
#rnx

diff --git a/Engine/Source/Runtime/Core/Public/UObject/ScriptDelegates.h b/Engine/Source/Runtime/Core/Public/UObject/ScriptDelegates.h
--- a/Engine/Source/Runtime/Core/Public/UObject/ScriptDelegates.h
+++ b/Engine/Source/Runtime/Core/Public/UObject/ScriptDelegates.h
@@ -8,6 +8,7 @@
 #include "Concepts/Scalar.h"
 #include "Containers/Array.h"
 #include "Containers/ContainerAllocationPolicies.h"
+#include "Containers/EnumAsByte.h"
 #include "Containers/UnrealString.h"
 #include "PropertyPortFlags.h"
 #include "Delegates/DelegateAccessHandler.h"
@@ -242,7 +243,11 @@ namespace UE::Core::Private
 	{
 		using MemberType = std::remove_cv_t<std::remove_reference_t<T>>;
 
-		MemberType Member;
+		// HACK - this forces non-scoped enums to be bytes.  This is expected in dynamic delegates and should be rejected by UHT,
+		// but hasn't been, though it did force them to bytes in its generated Params struct passed to ProcessEvent.
+		using RemappedMemberType = std::conditional_t<std::is_enum_v<MemberType> && std::is_convertible_v<MemberType, int>, TEnumAsByte<MemberType>, MemberType>;
+
+		RemappedMemberType Member;
 
 		// This odd constraint is to allow selection of the correct MoveOut overload by specifying the
 		// index when they are merged into the derived class's overload set.
