Core: Optimized TSizedInlineAllocator to reduce code and debug size

These changes are intended to offset the increases caused by the recent addition of SupportsElementAlignment support in this allocator. In a large monolithic server build for Linux, this reduced code size by ~24 KiB, .debug size by ~9.2 MiB, .psym size by ~1.3 MiB, and .sym size by ~1.8 MiB.

#Synced-CL 51935711

diff --git a/Engine/Source/Runtime/Core/Public/Containers/ContainerAllocationPolicies.h b/Engine/Source/Runtime/Core/Public/Containers/ContainerAllocationPolicies.h
--- a/Engine/Source/Runtime/Core/Public/Containers/ContainerAllocationPolicies.h
+++ b/Engine/Source/Runtime/Core/Public/Containers/ContainerAllocationPolicies.h
@@ -945,7 +945,7 @@ public:
 		}
 
 		// FContainerAllocatorInterface
-		UE_FORCEINLINE_HINT ElementType* GetAllocation() const
+		UE_NODEBUG UE_FORCEINLINE_HINT ElementType* GetAllocation() const
 		{
 			if (ElementType* Result = SecondaryData.GetAllocation())
 			{
@@ -955,6 +955,7 @@ public:
 		}
 
 		void ResizeAllocation(SizeType CurrentNum, SizeType NewMax, SIZE_T NumBytesPerElement)
+			requires (!TAllocatorTraits<SecondaryAllocator>::SupportsElementAlignment)
 		{
 			// Make sure the number of live elements is still withing the allocation since we only memmove, not destruct
 			checkSlow(CurrentNum <= NewMax);
@@ -989,7 +990,8 @@ public:
 			}
 		}
 
-		void ResizeAllocation(SizeType CurrentNum, SizeType NewMax, SIZE_T NumBytesPerElement, uint32 AlignmentOfElement)
+		void ResizeAllocation(SizeType CurrentNum, SizeType NewMax, SIZE_T NumBytesPerElement, uint32 AlignmentOfElement = DEFAULT_ALIGNMENT)
+			requires (!!TAllocatorTraits<SecondaryAllocator>::SupportsElementAlignment)
 		{
 			// Make sure the number of live elements is still withing the allocation since we only memmove, not destruct
 			checkSlow(CurrentNum <= NewMax);
@@ -1024,35 +1026,40 @@ public:
 			}
 		}
 
-		UE_FORCEINLINE_HINT SizeType CalculateSlackReserve(SizeType NewMax, SIZE_T NumBytesPerElement) const
+		UE_NODEBUG UE_FORCEINLINE_HINT SizeType CalculateSlackReserve(SizeType NewMax, SIZE_T NumBytesPerElement) const
+			requires (!TAllocatorTraits<SecondaryAllocator>::SupportsElementAlignment)
 		{
 			// If the elements use less space than the inline allocation, only use the inline allocation as slack.
 			return NewMax <= NumInlineElements ?
 				NumInlineElements :
 				SecondaryData.CalculateSlackReserve(NewMax, NumBytesPerElement);
 		}
-		UE_FORCEINLINE_HINT SizeType CalculateSlackReserve(SizeType NewMax, SIZE_T NumBytesPerElement, uint32 AlignmentOfElement) const
+		UE_NODEBUG UE_FORCEINLINE_HINT SizeType CalculateSlackReserve(SizeType NewMax, SIZE_T NumBytesPerElement, uint32 AlignmentOfElement = DEFAULT_ALIGNMENT) const
+			requires (!!TAllocatorTraits<SecondaryAllocator>::SupportsElementAlignment)
 		{
 			// If the elements use less space than the inline allocation, only use the inline allocation as slack.
 			return NewMax <= NumInlineElements ?
 				NumInlineElements :
 				SecondaryData.CalculateSlackReserve(NewMax, NumBytesPerElement, AlignmentOfElement);
 		}
-		UE_FORCEINLINE_HINT SizeType CalculateSlackShrink(SizeType NewMax, SizeType CurrentMax, SIZE_T NumBytesPerElement) const
+		UE_NODEBUG UE_FORCEINLINE_HINT SizeType CalculateSlackShrink(SizeType NewMax, SizeType CurrentMax, SIZE_T NumBytesPerElement) const
+			requires (!TAllocatorTraits<SecondaryAllocator>::SupportsElementAlignment)
 		{
 			// If the elements use less space than the inline allocation, only use the inline allocation as slack.
 			return NewMax <= NumInlineElements ?
 				NumInlineElements :
 				SecondaryData.CalculateSlackShrink(NewMax, CurrentMax, NumBytesPerElement);
 		}
-		UE_FORCEINLINE_HINT SizeType CalculateSlackShrink(SizeType NewMax, SizeType CurrentMax, SIZE_T NumBytesPerElement, uint32 AlignmentOfElement) const
+		UE_NODEBUG UE_FORCEINLINE_HINT SizeType CalculateSlackShrink(SizeType NewMax, SizeType CurrentMax, SIZE_T NumBytesPerElement, uint32 AlignmentOfElement = DEFAULT_ALIGNMENT) const
+			requires (!!TAllocatorTraits<SecondaryAllocator>::SupportsElementAlignment)
 		{
 			// If the elements use less space than the inline allocation, only use the inline allocation as slack.
 			return NewMax <= NumInlineElements ?
 				NumInlineElements :
 				SecondaryData.CalculateSlackShrink(NewMax, CurrentMax, NumBytesPerElement, AlignmentOfElement);
 		}
-		UE_FORCEINLINE_HINT SizeType CalculateSlackGrow(SizeType NewMax, SizeType CurrentMax, SIZE_T NumBytesPerElement) const
+		UE_NODEBUG UE_FORCEINLINE_HINT SizeType CalculateSlackGrow(SizeType NewMax, SizeType CurrentMax, SIZE_T NumBytesPerElement) const
+			requires (!TAllocatorTraits<SecondaryAllocator>::SupportsElementAlignment)
 		{
 			// If the elements use less space than the inline allocation, only use the inline allocation as slack.
 			// Also, when computing slack growth, don't count inline elements -- the slack algorithm has a special
@@ -1062,7 +1069,8 @@ public:
 				NumInlineElements :
 				SecondaryData.CalculateSlackGrow(NewMax, CurrentMax <= NumInlineElements ? 0 : CurrentMax, NumBytesPerElement);
 		}
-		UE_FORCEINLINE_HINT SizeType CalculateSlackGrow(SizeType NewMax, SizeType CurrentMax, SIZE_T NumBytesPerElement, uint32 AlignmentOfElement) const
+		UE_NODEBUG UE_FORCEINLINE_HINT SizeType CalculateSlackGrow(SizeType NewMax, SizeType CurrentMax, SIZE_T NumBytesPerElement, uint32 AlignmentOfElement = DEFAULT_ALIGNMENT) const
+			requires (!!TAllocatorTraits<SecondaryAllocator>::SupportsElementAlignment)
 		{
 			// If the elements use less space than the inline allocation, only use the inline allocation as slack.
 			// Also, when computing slack growth, don't count inline elements -- the slack algorithm has a special
@@ -1073,7 +1081,7 @@ public:
 				SecondaryData.CalculateSlackGrow(NewMax, CurrentMax <= NumInlineElements ? 0 : CurrentMax, NumBytesPerElement, AlignmentOfElement);
 		}
 
-		SIZE_T GetAllocatedSize(SizeType CurrentMax, SIZE_T NumBytesPerElement) const
+		UE_NODEBUG UE_FORCEINLINE_HINT SIZE_T GetAllocatedSize(SizeType CurrentMax, SIZE_T NumBytesPerElement) const
 		{
 			if (CurrentMax > NumInlineElements)
 			{
@@ -1082,7 +1090,7 @@ public:
 			return 0;
 		}
 
-		bool HasAllocation() const
+		UE_REWRITE bool HasAllocation() const
 		{
 			return SecondaryData.HasAllocation();
 		}
@@ -1113,7 +1121,7 @@ public:
 		typename SecondaryAllocator::template ForElementType<ElementType> SecondaryData;
 
 		/** @return the base of the aligned inline element data */
-		ElementType* GetInlineElements() const
+		UE_REWRITE ElementType* GetInlineElements() const
 		{
 			return (ElementType*)InlineData;
 		}
