As a memory usage optimization, added preprocessor define UE_WITH_CVAR_HELP_TEXT to control inclusion of cvar help text strings.

Default is UE_WITH_CVAR_HELP_TEXT=1.
With UE_WITH_CVAR_HELP_TEXT=0, savings are a few MB in program size and runtime category LLM/EngineMisc/ConsoleCommands.

#tests PIE, win64 replayrun, horde program size report, horde engine presubmit

#Synced-CL 51935769

diff --git a/Engine/Source/Runtime/Core/Private/HAL/ConsoleManager.cpp b/Engine/Source/Runtime/Core/Private/HAL/ConsoleManager.cpp
--- a/Engine/Source/Runtime/Core/Private/HAL/ConsoleManager.cpp
+++ b/Engine/Source/Runtime/Core/Private/HAL/ConsoleManager.cpp
@@ -206,6 +206,7 @@ public:
 
 	// interface IConsoleVariable -----------------------------------
 
+#if UE_WITH_CVAR_HELP_TEXT
 	virtual const TCHAR* GetHelp() const
 	{
 		return *Help;
@@ -227,6 +228,7 @@ public:
 		const FString CurrentValueString = GetString();
 		return FText::Format(LOCTEXT("CVarDetailedHelp", "{0}\n\nCurrent value : {1}"), FText::FromString(GetHelp()), CurrentValueString.IsEmpty() ? LOCTEXT("CVarEmptyValueString", "<empty>") : FText::FromString(CurrentValueString));
 	}
+#endif
 
 	virtual EConsoleVariableFlags GetFlags() const
 	{
@@ -347,9 +349,10 @@ public:
 #endif
 
 protected: // -----------------------------------------
-
+#if UE_WITH_CVAR_HELP_TEXT
 	// not using TCHAR* to allow chars support reloading of modules (otherwise we would keep a pointer into the module)
 	FString Help;
+#endif
 	//
 	EConsoleVariableFlags Flags;
 	/** User function to call when the console variable is changed */
@@ -396,12 +399,16 @@ public:
 	 * Constructor
 	 * @param InHelp must not be 0, must not be empty
 	 */
-	FConsoleCommandBase(const TCHAR* InHelp, EConsoleVariableFlags InFlags)
-		: Help(InHelp), Flags(InFlags)
+	FConsoleCommandBase(const TCHAR* InHelp, EConsoleVariableFlags InFlags) :
+#if UE_WITH_CVAR_HELP_TEXT
+		Help(InHelp),
+#endif
+		Flags(InFlags)
 	{
+#if UE_WITH_CVAR_HELP_TEXT
 		check(InHelp);
 		//check(*Help != 0); for now disabled as there callstack when we crash early during engine init
-
+#endif
 		ApplyPreviewIfScalability();
 	}
 
@@ -416,6 +423,7 @@ public:
 
 	// interface IConsoleVariable -----------------------------------
 
+#if UE_WITH_CVAR_HELP_TEXT
 	virtual const TCHAR* GetHelp() const
 	{
 		return *Help;
@@ -428,6 +436,7 @@ public:
 
 		Help = InValue;
 	}
+#endif
 
 	virtual EConsoleVariableFlags GetFlags() const
 	{
@@ -446,10 +455,10 @@ public:
 	}
 
 private: // -----------------------------------------
-
+#if UE_WITH_CVAR_HELP_TEXT
 	// not using TCHAR* to allow chars support reloading of modules (otherwise we would keep a pointer into the module)
 	FString Help;
-
+#endif
 	EConsoleVariableFlags Flags;
 };
 
@@ -2109,13 +2118,13 @@ public:
 };
 
 template<>
-TAutoRangedConsoleVariable<int32>::TAutoRangedConsoleVariable(const TCHAR* Name, int32 DefaultValue, int32 MinValue, int32 MaxValue, const TCHAR* Help, uint32 Flags, IConsoleVariable* Parent)
+TAutoRangedConsoleVariable<int32>::TAutoRangedConsoleVariable(IConsoleVariable* Parent, const TCHAR* Name, int32 DefaultValue, int32 MinValue, int32 MaxValue, const TCHAR* Help, uint32 Flags)
 	: TAutoConsoleVariable(GetManager().AddConsoleObject(Name, new FRangedConsoleVariable(Name, DefaultValue, MinValue, MaxValue, Help, (EConsoleVariableFlags)Flags, Parent))->AsVariable())
 {
 }
 
 template<>
-TAutoRangedConsoleVariable<float>::TAutoRangedConsoleVariable(const TCHAR* Name, float DefaultValue, float MinValue, float MaxValue, const TCHAR* Help, uint32 Flags, IConsoleVariable* Parent)
+TAutoRangedConsoleVariable<float>::TAutoRangedConsoleVariable(IConsoleVariable* Parent, const TCHAR* Name, float DefaultValue, float MinValue, float MaxValue, const TCHAR* Help, uint32 Flags)
 	: TAutoConsoleVariable(GetManager().AddConsoleObject(Name, new FRangedConsoleVariable(Name, DefaultValue, MinValue, MaxValue, Help, (EConsoleVariableFlags)Flags, Parent))->AsVariable())
 {
 }
@@ -5036,7 +5045,7 @@ public:
 		return nullptr;
 	}
 
-
+#if UE_WITH_CVAR_HELP_TEXT
 	virtual const TCHAR* GetHelp() const override
 	{
 		if (Bind())
@@ -5062,6 +5071,7 @@ public:
 			return RealVariable->SetHelp(Value);
 		}
 	}
+#endif
 
 	virtual EConsoleVariableFlags GetFlags() const override
 	{
@@ -5281,6 +5291,7 @@ public:
 		return bIsDeprecated;
 	}
 
+#if UE_WITH_CVAR_HELP_TEXT
 	virtual const TCHAR* GetHelp() const override
 	{
 		if (Bind())
@@ -5306,6 +5317,7 @@ public:
 		}
 		return INVTEXT("");
 	}
+#endif
 
 	virtual EConsoleVariableFlags GetFlags() const override
 	{
@@ -5492,15 +5504,15 @@ IConsoleVariable* CreateInPlaceDeprecatedVar(const FCVarDeprecationArgs& Depreca
 }
 
 template<typename T>
-TAutoConsoleVariable<T>::TAutoConsoleVariable(FCVarDeprecationArgs DeprecationArgs, const TCHAR* Name, const T& DefaultValue, const TCHAR* Help, uint32 Flags)
+TAutoConsoleVariable<T>::TAutoConsoleVariable(const FCVarDeprecationArgs& DeprecationArgs, const TCHAR* Name, const T& DefaultValue, uint32 Flags)
 	: TAutoConsoleVariable(CreateInPlaceDeprecatedVar(DeprecationArgs, Name, DefaultValue, Flags))
 {
-}	
+}
 
-template CORE_API TAutoConsoleVariable<int32>::TAutoConsoleVariable(FCVarDeprecationArgs DeprecationArgs, const TCHAR* Name, const int32& DefaultValue, const TCHAR* Help, uint32 Flags);
-template CORE_API TAutoConsoleVariable<float>::TAutoConsoleVariable(FCVarDeprecationArgs DeprecationArgs, const TCHAR* Name, const float& DefaultValue, const TCHAR* Help, uint32 Flags);
-template CORE_API TAutoConsoleVariable<bool>::TAutoConsoleVariable(FCVarDeprecationArgs DeprecationArgs, const TCHAR* Name, const bool& DefaultValue, const TCHAR* Help, uint32 Flags);
-template CORE_API TAutoConsoleVariable<FString>::TAutoConsoleVariable(FCVarDeprecationArgs DeprecationArgs, const TCHAR* Name, const FString& DefaultValue, const TCHAR* Help, uint32 Flags);
+template CORE_API TAutoConsoleVariable<int32>::TAutoConsoleVariable(const FCVarDeprecationArgs& DeprecationArgs, const TCHAR* Name, const int32& DefaultValue, uint32 Flags);
+template CORE_API TAutoConsoleVariable<float>::TAutoConsoleVariable(const FCVarDeprecationArgs& DeprecationArgs, const TCHAR* Name, const float& DefaultValue, uint32 Flags);
+template CORE_API TAutoConsoleVariable<bool>::TAutoConsoleVariable(const FCVarDeprecationArgs& DeprecationArgs, const TCHAR* Name, const bool& DefaultValue, uint32 Flags);
+template CORE_API TAutoConsoleVariable<FString>::TAutoConsoleVariable(const FCVarDeprecationArgs& DeprecationArgs, const TCHAR* Name, const FString& DefaultValue, uint32 Flags);
 
 template<typename T, typename CmdType>
 void CreateInPlaceDeprecatedCommand(const TCHAR* Name, const CmdType& Command, const TCHAR* DeprecatedAtVersion, EShadowCVarBehavior UsageBehavior, const TCHAR* DeprecationMessage)
@@ -5636,7 +5648,7 @@ namespace ConsoleVariablePlatformMapping
 #if !NO_CVARS
 
 template<typename T>
-IConsoleVariable* CreateInPlaceDeprecatedVarRef(FCVarDeprecationArgs DeprecationArgs, const TCHAR* Name, T& RefValue, const TCHAR* Help, uint32 Flags)
+IConsoleVariable* CreateInPlaceDeprecatedVarRef(const FCVarDeprecationArgs& DeprecationArgs, const TCHAR* Name, T& RefValue, uint32 Flags)
 {
 	return CreateInPlaceDeprecatedVar<T>(DeprecationArgs, Name, Flags, 
 		[&](const FString& DeprecationHelp, EConsoleVariableFlags ConsoleVariableFlags)
@@ -5645,192 +5657,199 @@ IConsoleVariable* CreateInPlaceDeprecatedVarRef(FCVarDeprecationArgs Deprecation
 	});
 }
 
-FAutoConsoleVariable::FAutoConsoleVariable(const TCHAR* Name, bool DefaultValue, const TCHAR* Help, uint32 Flags)
+FAutoConsoleVariable::FAutoConsoleVariable(uint32 Flags, const TCHAR* Name, bool DefaultValue, const TCHAR* Help)
 	: FAutoConsoleObject(IConsoleManager::Get().RegisterConsoleVariable(Name, DefaultValue, Help, Flags))
 {
 }
 
-FAutoConsoleVariable::FAutoConsoleVariable(const TCHAR* Name, int32 DefaultValue, const TCHAR* Help, uint32 Flags)
+FAutoConsoleVariable::FAutoConsoleVariable(uint32 Flags, const TCHAR* Name, int32 DefaultValue, const TCHAR* Help)
 	: FAutoConsoleObject(IConsoleManager::Get().RegisterConsoleVariable(Name, DefaultValue, Help, Flags))
 {
 }
 
-FAutoConsoleVariable::FAutoConsoleVariable(const TCHAR* Name, float DefaultValue, const TCHAR* Help, uint32 Flags)
+FAutoConsoleVariable::FAutoConsoleVariable(uint32 Flags, const TCHAR* Name, float DefaultValue, const TCHAR* Help)
 	: FAutoConsoleObject(IConsoleManager::Get().RegisterConsoleVariable(Name, DefaultValue, Help, Flags))
 {
 }
 
-FAutoConsoleVariable::FAutoConsoleVariable(const TCHAR* Name, const TCHAR* DefaultValue, const TCHAR* Help, uint32 Flags)
+FAutoConsoleVariable::FAutoConsoleVariable(uint32 Flags, const TCHAR* Name, const TCHAR* DefaultValue, const TCHAR* Help)
 	: FAutoConsoleObject(IConsoleManager::Get().RegisterConsoleVariable(Name, DefaultValue, Help, Flags))
 {
 }
 
-FAutoConsoleVariable::FAutoConsoleVariable(const TCHAR* Name, bool DefaultValue, const TCHAR* Help, const FConsoleVariableDelegate& Callback, uint32 Flags)
+FAutoConsoleVariable::FAutoConsoleVariable(uint32 Flags, const TCHAR* Name, bool DefaultValue, const TCHAR* Help, const FConsoleVariableDelegate& Callback)
 	: FAutoConsoleObject(IConsoleManager::Get().RegisterConsoleVariable(Name, DefaultValue, Help, Flags))
 {
 	AsVariable()->SetOnChangedCallback(Callback);
 }
 
-FAutoConsoleVariable::FAutoConsoleVariable(const TCHAR* Name, int32 DefaultValue, const TCHAR* Help, const FConsoleVariableDelegate& Callback, uint32 Flags)
+FAutoConsoleVariable::FAutoConsoleVariable(uint32 Flags, const TCHAR* Name, int32 DefaultValue, const TCHAR* Help, const FConsoleVariableDelegate& Callback)
 	: FAutoConsoleObject(IConsoleManager::Get().RegisterConsoleVariable(Name, DefaultValue, Help, Flags))
 {
 	AsVariable()->SetOnChangedCallback(Callback);
 }
 
-FAutoConsoleVariable::FAutoConsoleVariable(const TCHAR* Name, float DefaultValue, const TCHAR* Help, const FConsoleVariableDelegate& Callback, uint32 Flags)
+FAutoConsoleVariable::FAutoConsoleVariable(uint32 Flags, const TCHAR* Name, float DefaultValue, const TCHAR* Help, const FConsoleVariableDelegate& Callback)
 	: FAutoConsoleObject(IConsoleManager::Get().RegisterConsoleVariable(Name, DefaultValue, Help, Flags))
 {
 	AsVariable()->SetOnChangedCallback(Callback);
 }
 
-FAutoConsoleVariable::FAutoConsoleVariable(const TCHAR* Name, const TCHAR* DefaultValue, const TCHAR* Help, const FConsoleVariableDelegate& Callback, uint32 Flags)
+FAutoConsoleVariable::FAutoConsoleVariable(uint32 Flags, const TCHAR* Name, const TCHAR* DefaultValue, const TCHAR* Help, const FConsoleVariableDelegate& Callback)
 	: FAutoConsoleObject(IConsoleManager::Get().RegisterConsoleVariable(Name, DefaultValue, Help, Flags))
 {
 	AsVariable()->SetOnChangedCallback(Callback);
 }
 
-FAutoConsoleVariableRef::FAutoConsoleVariableRef(const TCHAR* Name, int32& RefValue, const TCHAR* Help, uint32 Flags)
+FAutoConsoleVariableRef::FAutoConsoleVariableRef(uint32 Flags, const TCHAR* Name, int32& RefValue, const TCHAR* Help)
 	: FAutoConsoleObject(IConsoleManager::Get().RegisterConsoleVariableRef(Name, RefValue, Help, Flags))
 {
 }
 
-FAutoConsoleVariableRef::FAutoConsoleVariableRef(FCVarDeprecationArgs DeprecationArgs, const TCHAR* Name, int32& RefValue, const TCHAR* Help, uint32 Flags)
-	: FAutoConsoleObject(CreateInPlaceDeprecatedVarRef(DeprecationArgs, Name, RefValue, Help, Flags))
+FAutoConsoleVariableRef::FAutoConsoleVariableRef(const FCVarDeprecationArgs& DeprecationArgs, const TCHAR* Name, int32& RefValue, uint32 Flags)
+	: FAutoConsoleObject(CreateInPlaceDeprecatedVarRef(DeprecationArgs, Name, RefValue, Flags))
 {
 }
 
-FAutoConsoleVariableRef::FAutoConsoleVariableRef(const TCHAR* Name, float& RefValue, const TCHAR* Help, uint32 Flags)
+FAutoConsoleVariableRef::FAutoConsoleVariableRef(uint32 Flags, const TCHAR* Name, float& RefValue, const TCHAR* Help)
 	: FAutoConsoleObject(IConsoleManager::Get().RegisterConsoleVariableRef(Name, RefValue, Help, Flags))
 {
 }
 
-FAutoConsoleVariableRef::FAutoConsoleVariableRef(FCVarDeprecationArgs DeprecationArgs, const TCHAR* Name, float& RefValue, const TCHAR* Help, uint32 Flags)
-	: FAutoConsoleObject(CreateInPlaceDeprecatedVarRef(DeprecationArgs, Name, RefValue, Help, Flags))
+FAutoConsoleVariableRef::FAutoConsoleVariableRef(const FCVarDeprecationArgs& DeprecationArgs, const TCHAR* Name, float& RefValue, uint32 Flags)
+	: FAutoConsoleObject(CreateInPlaceDeprecatedVarRef(DeprecationArgs, Name, RefValue, Flags))
 {
 }
-FAutoConsoleVariableRef::FAutoConsoleVariableRef(const TCHAR* Name, bool& RefValue, const TCHAR* Help, uint32 Flags)
+
+FAutoConsoleVariableRef::FAutoConsoleVariableRef(uint32 Flags, const TCHAR* Name, bool& RefValue, const TCHAR* Help)
 	: FAutoConsoleObject(IConsoleManager::Get().RegisterConsoleVariableRef(Name, RefValue, Help, Flags))
 {
 }
 
-FAutoConsoleVariableRef::FAutoConsoleVariableRef(FCVarDeprecationArgs DeprecationArgs, const TCHAR* Name, bool& RefValue, const TCHAR* Help, uint32 Flags)
-	: FAutoConsoleObject(CreateInPlaceDeprecatedVarRef(DeprecationArgs, Name, RefValue, Help, Flags))
+FAutoConsoleVariableRef::FAutoConsoleVariableRef(const FCVarDeprecationArgs& DeprecationArgs, const TCHAR* Name, bool& RefValue, uint32 Flags)
+	: FAutoConsoleObject(CreateInPlaceDeprecatedVarRef(DeprecationArgs, Name, RefValue, Flags))
 {
 }
-FAutoConsoleVariableRef::FAutoConsoleVariableRef(const TCHAR* Name, FString& RefValue, const TCHAR* Help, uint32 Flags)
+
+FAutoConsoleVariableRef::FAutoConsoleVariableRef(uint32 Flags, const TCHAR* Name, FString& RefValue, const TCHAR* Help)
 	: FAutoConsoleObject(IConsoleManager::Get().RegisterConsoleVariableRef(Name, RefValue, Help, Flags))
 {
 }
 
-FAutoConsoleVariableRef::FAutoConsoleVariableRef(FCVarDeprecationArgs DeprecationArgs, const TCHAR* Name, FString& RefValue, const TCHAR* Help, uint32 Flags)
-	: FAutoConsoleObject(CreateInPlaceDeprecatedVarRef(DeprecationArgs, Name, RefValue, Help, Flags))
+FAutoConsoleVariableRef::FAutoConsoleVariableRef(const FCVarDeprecationArgs& DeprecationArgs, const TCHAR* Name, FString& RefValue, uint32 Flags)
+	: FAutoConsoleObject(CreateInPlaceDeprecatedVarRef(DeprecationArgs, Name, RefValue, Flags))
 {
 }
-FAutoConsoleVariableRef::FAutoConsoleVariableRef(const TCHAR* Name, FName& RefValue, const TCHAR* Help, uint32 Flags)
+
+FAutoConsoleVariableRef::FAutoConsoleVariableRef(uint32 Flags, const TCHAR* Name, FName& RefValue, const TCHAR* Help)
 	: FAutoConsoleObject(IConsoleManager::Get().RegisterConsoleVariableRef(Name, RefValue, Help, Flags))
 {
 }
 
-FAutoConsoleVariableRef::FAutoConsoleVariableRef(FCVarDeprecationArgs DeprecationArgs, const TCHAR* Name, FName& RefValue, const TCHAR* Help, uint32 Flags)
-	: FAutoConsoleObject(CreateInPlaceDeprecatedVarRef(DeprecationArgs, Name, RefValue, Help, Flags))
+FAutoConsoleVariableRef::FAutoConsoleVariableRef(const FCVarDeprecationArgs& DeprecationArgs, const TCHAR* Name, FName& RefValue, uint32 Flags)
+	: FAutoConsoleObject(CreateInPlaceDeprecatedVarRef(DeprecationArgs, Name, RefValue, Flags))
 {
 }
-FAutoConsoleVariableRef::FAutoConsoleVariableRef(const TCHAR* Name, int32& RefValue, const TCHAR* Help, const FConsoleVariableDelegate& Callback, uint32 Flags)
+
+FAutoConsoleVariableRef::FAutoConsoleVariableRef(uint32 Flags, const TCHAR* Name, int32& RefValue, const TCHAR* Help, const FConsoleVariableDelegate& Callback)
 	: FAutoConsoleObject(IConsoleManager::Get().RegisterConsoleVariableRef(Name, RefValue, Help, Flags))
 {
 	AsVariable()->SetOnChangedCallback(Callback);
 }
 
-FAutoConsoleVariableRef::FAutoConsoleVariableRef(FCVarDeprecationArgs DeprecationArgs, const TCHAR* Name, int32& RefValue, const TCHAR* Help, const FConsoleVariableDelegate& Callback, uint32 Flags)
-	: FAutoConsoleObject(CreateInPlaceDeprecatedVarRef(DeprecationArgs, Name, RefValue, Help, Flags))
+FAutoConsoleVariableRef::FAutoConsoleVariableRef(const FCVarDeprecationArgs& DeprecationArgs, const TCHAR* Name, int32& RefValue, const FConsoleVariableDelegate& Callback, uint32 Flags)
+	: FAutoConsoleObject(CreateInPlaceDeprecatedVarRef(DeprecationArgs, Name, RefValue, Flags))
 {
 	AsVariable()->SetOnChangedCallback(Callback);
 }
 
-FAutoConsoleVariableRef::FAutoConsoleVariableRef(const TCHAR* Name, float& RefValue, const TCHAR* Help, const FConsoleVariableDelegate& Callback, uint32 Flags)
+FAutoConsoleVariableRef::FAutoConsoleVariableRef(uint32 Flags, const TCHAR* Name, float& RefValue, const TCHAR* Help, const FConsoleVariableDelegate& Callback)
 	: FAutoConsoleObject(IConsoleManager::Get().RegisterConsoleVariableRef(Name, RefValue, Help, Flags))
 {
 	AsVariable()->SetOnChangedCallback(Callback);
 }
-FAutoConsoleVariableRef::FAutoConsoleVariableRef(FCVarDeprecationArgs DeprecationArgs, const TCHAR* Name, float& RefValue, const TCHAR* Help, const FConsoleVariableDelegate& Callback, uint32 Flags)
-	: FAutoConsoleObject(CreateInPlaceDeprecatedVarRef(DeprecationArgs, Name, RefValue, Help, Flags))
+
+FAutoConsoleVariableRef::FAutoConsoleVariableRef(const FCVarDeprecationArgs& DeprecationArgs, const TCHAR* Name, float& RefValue, const FConsoleVariableDelegate& Callback, uint32 Flags)
+	: FAutoConsoleObject(CreateInPlaceDeprecatedVarRef(DeprecationArgs, Name, RefValue, Flags))
 {
 	AsVariable()->SetOnChangedCallback(Callback);
 }
 
-FAutoConsoleVariableRef::FAutoConsoleVariableRef(const TCHAR* Name, bool& RefValue, const TCHAR* Help, const FConsoleVariableDelegate& Callback, uint32 Flags)
+FAutoConsoleVariableRef::FAutoConsoleVariableRef(uint32 Flags, const TCHAR* Name, bool& RefValue, const TCHAR* Help, const FConsoleVariableDelegate& Callback)
 	: FAutoConsoleObject(IConsoleManager::Get().RegisterConsoleVariableRef(Name, RefValue, Help, Flags))
 {
 	AsVariable()->SetOnChangedCallback(Callback);
 }
-FAutoConsoleVariableRef::FAutoConsoleVariableRef(FCVarDeprecationArgs DeprecationArgs, const TCHAR* Name, bool& RefValue, const TCHAR* Help, const FConsoleVariableDelegate& Callback, uint32 Flags)
-	: FAutoConsoleObject(CreateInPlaceDeprecatedVarRef(DeprecationArgs, Name, RefValue, Help, Flags))
+
+FAutoConsoleVariableRef::FAutoConsoleVariableRef(const FCVarDeprecationArgs& DeprecationArgs, const TCHAR* Name, bool& RefValue, const FConsoleVariableDelegate& Callback, uint32 Flags)
+	: FAutoConsoleObject(CreateInPlaceDeprecatedVarRef(DeprecationArgs, Name, RefValue, Flags))
 {
 	AsVariable()->SetOnChangedCallback(Callback);
 }
 
-FAutoConsoleVariableRef::FAutoConsoleVariableRef(const TCHAR* Name, FString& RefValue, const TCHAR* Help, const FConsoleVariableDelegate& Callback, uint32 Flags)
+FAutoConsoleVariableRef::FAutoConsoleVariableRef(uint32 Flags, const TCHAR* Name, FString& RefValue, const TCHAR* Help, const FConsoleVariableDelegate& Callback)
 	: FAutoConsoleObject(IConsoleManager::Get().RegisterConsoleVariableRef(Name, RefValue, Help, Flags))
 {
 	AsVariable()->SetOnChangedCallback(Callback);
 }
-FAutoConsoleVariableRef::FAutoConsoleVariableRef(FCVarDeprecationArgs DeprecationArgs, const TCHAR* Name, FString& RefValue, const TCHAR* Help, const FConsoleVariableDelegate& Callback, uint32 Flags)
-	: FAutoConsoleObject(CreateInPlaceDeprecatedVarRef(DeprecationArgs, Name, RefValue, Help, Flags))
+
+FAutoConsoleVariableRef::FAutoConsoleVariableRef(const FCVarDeprecationArgs& DeprecationArgs, const TCHAR* Name, FString& RefValue, const FConsoleVariableDelegate& Callback, uint32 Flags)
+	: FAutoConsoleObject(CreateInPlaceDeprecatedVarRef(DeprecationArgs, Name, RefValue, Flags))
 {
 	AsVariable()->SetOnChangedCallback(Callback);
 }
 
-FAutoConsoleCommand::FAutoConsoleCommand(const TCHAR* Name, const TCHAR* Help, const FConsoleCommandDelegate& Command, uint32 Flags)
+FAutoConsoleCommand::FAutoConsoleCommand(uint32 Flags, const TCHAR* Name, const TCHAR* Help, const FConsoleCommandDelegate& Command)
 	: FAutoConsoleObject(IConsoleManager::Get().RegisterConsoleCommand(Name, Help, Command, Flags))
 {
 }
 
-FAutoConsoleCommand::FAutoConsoleCommand(const TCHAR* Name, const TCHAR* Help, const FConsoleCommandWithArgsDelegate& Command, uint32 Flags)
+FAutoConsoleCommand::FAutoConsoleCommand(uint32 Flags, const TCHAR* Name, const TCHAR* Help, const FConsoleCommandWithArgsDelegate& Command)
 	: FAutoConsoleObject(IConsoleManager::Get().RegisterConsoleCommand(Name, Help, Command, Flags))
 {
 }
 
-FAutoConsoleCommand::FAutoConsoleCommand(const TCHAR* Name, const TCHAR* Help, const FConsoleCommandWithOutputDeviceDelegate& Command, uint32 Flags)
+FAutoConsoleCommand::FAutoConsoleCommand(uint32 Flags, const TCHAR* Name, const TCHAR* Help, const FConsoleCommandWithOutputDeviceDelegate& Command)
 	: FAutoConsoleObject(IConsoleManager::Get().RegisterConsoleCommand(Name, Help, Command, Flags))
 {
 }
 
-FAutoConsoleCommand::FAutoConsoleCommand(const TCHAR* Name, const TCHAR* Help, const FConsoleCommandWithArgsAndOutputDeviceDelegate& Command, uint32 Flags)
+FAutoConsoleCommand::FAutoConsoleCommand(uint32 Flags, const TCHAR* Name, const TCHAR* Help, const FConsoleCommandWithArgsAndOutputDeviceDelegate& Command)
 	: FAutoConsoleObject(IConsoleManager::Get().RegisterConsoleCommand(Name, Help, Command, Flags))
 {
 }
 
-FAutoConsoleCommand::FAutoConsoleCommand(const TCHAR* Name, const TCHAR* Help, const FConsoleCommandWithWorldDelegate& Command, uint32 Flags)
+FAutoConsoleCommand::FAutoConsoleCommand(uint32 Flags, const TCHAR* Name, const TCHAR* Help, const FConsoleCommandWithWorldDelegate& Command)
 	: FAutoConsoleObject(IConsoleManager::Get().RegisterConsoleCommand(Name, Help, Command, Flags))
 {
 }
 
-FAutoConsoleCommand::FAutoConsoleCommand(const TCHAR* Name, const TCHAR* Help, const FConsoleCommandWithWorldArgsAndOutputDeviceDelegate& Command, uint32 Flags)
+FAutoConsoleCommand::FAutoConsoleCommand(uint32 Flags, const TCHAR* Name, const TCHAR* Help, const FConsoleCommandWithWorldArgsAndOutputDeviceDelegate& Command)
 	: FAutoConsoleObject(IConsoleManager::Get().RegisterConsoleCommand(Name, Help, Command, Flags))
 {
 }
 
-FAutoConsoleCommandWithWorld::FAutoConsoleCommandWithWorld(const TCHAR* Name, const TCHAR* Help, const FConsoleCommandWithWorldDelegate& Command, uint32 Flags)
+FAutoConsoleCommandWithWorld::FAutoConsoleCommandWithWorld(uint32 Flags, const TCHAR* Name, const TCHAR* Help, const FConsoleCommandWithWorldDelegate& Command)
 	: FAutoConsoleObject(IConsoleManager::Get().RegisterConsoleCommand(Name, Help, Command, Flags))
 {
 }
 
-FAutoConsoleCommandWithWorldAndArgs::FAutoConsoleCommandWithWorldAndArgs(const TCHAR* Name, const TCHAR* Help, const FConsoleCommandWithWorldAndArgsDelegate& Command, uint32 Flags)
+FAutoConsoleCommandWithWorldAndArgs::FAutoConsoleCommandWithWorldAndArgs(uint32 Flags, const TCHAR* Name, const TCHAR* Help, const FConsoleCommandWithWorldAndArgsDelegate& Command)
 	: FAutoConsoleObject(IConsoleManager::Get().RegisterConsoleCommand(Name, Help, Command, Flags))
 {
 }
 
-FAutoConsoleCommandWithArgsAndOutputDevice::FAutoConsoleCommandWithArgsAndOutputDevice(const TCHAR* Name, const TCHAR* Help, const FConsoleCommandWithArgsAndOutputDeviceDelegate& Command, uint32 Flags)
+FAutoConsoleCommandWithArgsAndOutputDevice::FAutoConsoleCommandWithArgsAndOutputDevice(uint32 Flags, const TCHAR* Name, const TCHAR* Help, const FConsoleCommandWithArgsAndOutputDeviceDelegate& Command)
 	: FAutoConsoleObject(IConsoleManager::Get().RegisterConsoleCommand(Name, Help, Command, Flags))
 {
 }
 
-FAutoConsoleCommandWithOutputDevice::FAutoConsoleCommandWithOutputDevice(const TCHAR* Name, const TCHAR* Help, const FConsoleCommandWithOutputDeviceDelegate& Command, uint32 Flags)
+FAutoConsoleCommandWithOutputDevice::FAutoConsoleCommandWithOutputDevice(uint32 Flags, const TCHAR* Name, const TCHAR* Help, const FConsoleCommandWithOutputDeviceDelegate& Command)
 	: FAutoConsoleObject(IConsoleManager::Get().RegisterConsoleCommand(Name, Help, Command, Flags))
 {
 }
 
-FAutoConsoleCommandWithWorldArgsAndOutputDevice::FAutoConsoleCommandWithWorldArgsAndOutputDevice(const TCHAR* Name, const TCHAR* Help, const FConsoleCommandWithWorldArgsAndOutputDeviceDelegate& Command, uint32 Flags)
+FAutoConsoleCommandWithWorldArgsAndOutputDevice::FAutoConsoleCommandWithWorldArgsAndOutputDevice(uint32 Flags, const TCHAR* Name, const TCHAR* Help, const FConsoleCommandWithWorldArgsAndOutputDeviceDelegate& Command)
 	: FAutoConsoleObject(IConsoleManager::Get().RegisterConsoleCommand(Name, Help, Command, Flags))
 {
 }
diff --git a/Engine/Source/Runtime/Core/Public/HAL/IConsoleManager.h b/Engine/Source/Runtime/Core/Public/HAL/IConsoleManager.h
--- a/Engine/Source/Runtime/Core/Public/HAL/IConsoleManager.h
+++ b/Engine/Source/Runtime/Core/Public/HAL/IConsoleManager.h
@@ -14,6 +14,16 @@
 #include "Templates/EnableIf.h"
 #include "Internationalization/Text.h"
 
+#if !defined(UE_WITH_CVAR_HELP_TEXT)
+	#define UE_WITH_CVAR_HELP_TEXT 1
+#endif
+
+#if UE_WITH_CVAR_HELP_TEXT
+	#define UE_CVAR_HELP_TEXT_PARAM(HelpText) HelpText
+#else
+	#define UE_CVAR_HELP_TEXT_PARAM(HelpText) nullptr
+#endif
+
 #define TRACK_CONSOLE_FIND_COUNT !(UE_BUILD_SHIPPING || UE_BUILD_TEST)
 
 #if DO_CHECK && (!UE_BUILD_SHIPPING) && !defined(__clang_analyzer__) // Disable even if checks in shipping are enabled.
@@ -399,6 +409,7 @@ public:
 
 	virtual ~IConsoleObject() = default;
 
+#if UE_WITH_CVAR_HELP_TEXT
 	/**
 	 *  @return never 0, can be multi line ('\n')
 	 */
@@ -408,14 +419,6 @@ public:
 	 */
 	virtual void SetHelp(const TCHAR* Value) = 0;
 	/**
-	 * Get the internal state of the flags.
-	 */
-	virtual EConsoleVariableFlags GetFlags() const = 0;
-	/**
-	 * Sets the internal flag state to the specified value.
-	 */
-	virtual void SetFlags(const EConsoleVariableFlags Value) = 0;
-	/** 
 	 *  @return a (potentially) more detailed help than GetHelp (e.g. current value for console variables)
 	 */
 	virtual FText GetDetailedHelp() const
@@ -423,6 +426,20 @@ public:
 		// By default, just return the standard help message :
 		return FText::FromString(GetHelp());
 	}
+#else
+	const TCHAR* GetHelp() const { return TEXT(""); }
+	void SetHelp(const TCHAR*) {}
+	FText GetDetailedHelp() const { return FText::FromString(TEXT("")); }
+#endif
+
+	/**
+	 * Get the internal state of the flags.
+	 */
+	virtual EConsoleVariableFlags GetFlags() const = 0;
+	/**
+	 * Sets the internal flag state to the specified value.
+	 */
+	virtual void SetFlags(const EConsoleVariableFlags Value) = 0;
 
 	// Convenience methods -------------------------------------
 
@@ -1516,6 +1533,14 @@ private:
  */
 class FAutoConsoleVariable : private FAutoConsoleObject
 {
+	CORE_API FAutoConsoleVariable(uint32 Flags, const TCHAR* Name, bool DefaultValue, const TCHAR* Help);
+	CORE_API FAutoConsoleVariable(uint32 Flags, const TCHAR* Name, int32 DefaultValue, const TCHAR* Help);
+	CORE_API FAutoConsoleVariable(uint32 Flags, const TCHAR* Name, float DefaultValue, const TCHAR* Help);
+	CORE_API FAutoConsoleVariable(uint32 Flags, const TCHAR* Name, const TCHAR* DefaultValue, const TCHAR* Help);
+	CORE_API FAutoConsoleVariable(uint32 Flags, const TCHAR* Name, bool DefaultValue, const TCHAR* Help, const FConsoleVariableDelegate& Callback);
+	CORE_API FAutoConsoleVariable(uint32 Flags, const TCHAR* Name, int32 DefaultValue, const TCHAR* Help, const FConsoleVariableDelegate& Callback);
+	CORE_API FAutoConsoleVariable(uint32 Flags, const TCHAR* Name, float DefaultValue, const TCHAR* Help, const FConsoleVariableDelegate& Callback);
+	CORE_API FAutoConsoleVariable(uint32 Flags, const TCHAR* Name, const TCHAR* DefaultValue, const TCHAR* Help, const FConsoleVariableDelegate& Callback);
 public:
 	/**
 	 * Create a bool console variable
@@ -1523,29 +1548,32 @@ public:
 	 * @param Help must not be 0
 	 * @param Flags bitmask combined from EConsoleVariableFlags
 	 */
-	CORE_API FAutoConsoleVariable(const TCHAR* Name, bool DefaultValue, const TCHAR* Help, uint32 Flags = ECVF_Default);
+	FORCEINLINE FAutoConsoleVariable(const TCHAR* Name, bool DefaultValue, const TCHAR* Help, uint32 Flags = ECVF_Default)
+		: FAutoConsoleVariable(Flags, Name, DefaultValue, UE_CVAR_HELP_TEXT_PARAM(Help)) {}
 	/**
 	 * Create a int console variable
 	 * @param Name must not be 0
 	 * @param Help must not be 0
 	 * @param Flags bitmask combined from EConsoleVariableFlags
 	 */
-	CORE_API FAutoConsoleVariable(const TCHAR* Name, int32 DefaultValue, const TCHAR* Help, uint32 Flags = ECVF_Default);
+	FORCEINLINE FAutoConsoleVariable(const TCHAR* Name, int32 DefaultValue, const TCHAR* Help, uint32 Flags = ECVF_Default)
+		: FAutoConsoleVariable(Flags, Name, DefaultValue, UE_CVAR_HELP_TEXT_PARAM(Help)) {}
 	/**
 	 * Create a float console variable
 	 * @param Name must not be 0
 	 * @param Help must not be 0
 	 * @param Flags bitmask combined from EConsoleVariableFlags
 	 */
-	CORE_API FAutoConsoleVariable(const TCHAR* Name, float DefaultValue, const TCHAR* Help, uint32 Flags = ECVF_Default);
+	FORCEINLINE FAutoConsoleVariable(const TCHAR* Name, float DefaultValue, const TCHAR* Help, uint32 Flags = ECVF_Default)
+		: FAutoConsoleVariable(Flags, Name, DefaultValue, UE_CVAR_HELP_TEXT_PARAM(Help)) {}
 	/**
 	 * Create a string console variable
 	 * @param Name must not be 0
 	 * @param Help must not be 0
 	 * @param Flags bitmask combined from EConsoleVariableFlags
 	 */
-	CORE_API FAutoConsoleVariable(const TCHAR* Name, const TCHAR* DefaultValue, const TCHAR* Help, uint32 Flags = ECVF_Default);
-
+	FORCEINLINE FAutoConsoleVariable(const TCHAR* Name, const TCHAR* DefaultValue, const TCHAR* Help, uint32 Flags = ECVF_Default)
+		: FAutoConsoleVariable(Flags, Name, DefaultValue, UE_CVAR_HELP_TEXT_PARAM(Help)) {}
 	/**
 	 * Create a bool console variable
 	 * @param Name must not be 0
@@ -1553,8 +1581,8 @@ public:
 	 * @param Callback Delegate called when the variable changes. @see IConsoleVariable::SetOnChangedCallback
 	 * @param Flags bitmask combined from EConsoleVariableFlags
 	 */
-	CORE_API FAutoConsoleVariable(const TCHAR* Name, bool DefaultValue, const TCHAR* Help, const FConsoleVariableDelegate& Callback, uint32 Flags = ECVF_Default);
-	
+	FORCEINLINE FAutoConsoleVariable(const TCHAR* Name, bool DefaultValue, const TCHAR* Help, const FConsoleVariableDelegate& Callback, uint32 Flags = ECVF_Default)
+		: FAutoConsoleVariable(Flags, Name, DefaultValue, UE_CVAR_HELP_TEXT_PARAM(Help), Callback) {}
 	/**
 	 * Create a int console variable
 	 * @param Name must not be 0
@@ -1562,8 +1590,8 @@ public:
 	 * @param Callback Delegate called when the variable changes. @see IConsoleVariable::SetOnChangedCallback
 	 * @param Flags bitmask combined from EConsoleVariableFlags
 	 */
-	CORE_API FAutoConsoleVariable(const TCHAR* Name, int32 DefaultValue, const TCHAR* Help, const FConsoleVariableDelegate& Callback, uint32 Flags = ECVF_Default);
-
+	FORCEINLINE FAutoConsoleVariable(const TCHAR* Name, int32 DefaultValue, const TCHAR* Help, const FConsoleVariableDelegate& Callback, uint32 Flags = ECVF_Default)
+		: FAutoConsoleVariable(Flags, Name, DefaultValue, UE_CVAR_HELP_TEXT_PARAM(Help), Callback) {}
 	/**
 	 * Create a float console variable
 	 * @param Name must not be 0
@@ -1571,8 +1599,8 @@ public:
 	 * @param Callback Delegate called when the variable changes. @see IConsoleVariable::SetOnChangedCallback
 	 * @param Flags bitmask combined from EConsoleVariableFlags
 	 */
-	CORE_API FAutoConsoleVariable(const TCHAR* Name, float DefaultValue, const TCHAR* Help, const FConsoleVariableDelegate& Callback, uint32 Flags = ECVF_Default);
-
+	FORCEINLINE FAutoConsoleVariable(const TCHAR* Name, float DefaultValue, const TCHAR* Help, const FConsoleVariableDelegate& Callback, uint32 Flags = ECVF_Default)
+		: FAutoConsoleVariable(Flags, Name, DefaultValue, UE_CVAR_HELP_TEXT_PARAM(Help), Callback) {}
 	/**
 	 * Create a string console variable
 	 * @param Name must not be 0
@@ -1580,7 +1608,10 @@ public:
 	 * @param Callback Delegate called when the variable changes. @see IConsoleVariable::SetOnChangedCallback
 	 * @param Flags bitmask combined from EConsoleVariableFlags
 	 */
-	CORE_API FAutoConsoleVariable(const TCHAR* Name, const TCHAR* DefaultValue, const TCHAR* Help, const FConsoleVariableDelegate& Callback, uint32 Flags = ECVF_Default);
+	FORCEINLINE FAutoConsoleVariable(const TCHAR* Name, const TCHAR* DefaultValue, const TCHAR* Help, const FConsoleVariableDelegate& Callback, uint32 Flags = ECVF_Default)
+		: FAutoConsoleVariable(Flags, Name, DefaultValue, UE_CVAR_HELP_TEXT_PARAM(Help), Callback) {}
+
+	virtual ~FAutoConsoleVariable() {}
 
 	/** Dereference back to a console variable**/
 	UE_FORCEINLINE_HINT IConsoleVariable& operator*()
@@ -1669,6 +1700,26 @@ struct FCVarDeprecationArgs
  */
 class FAutoConsoleVariableRef : private FAutoConsoleObject
 {
+	CORE_API FAutoConsoleVariableRef(uint32 Flags, const TCHAR* Name, int32& RefValue, const TCHAR* Help);
+	CORE_API FAutoConsoleVariableRef(uint32 Flags, const TCHAR* Name, float& RefValue, const TCHAR* Help);
+	CORE_API FAutoConsoleVariableRef(uint32 Flags, const TCHAR* Name, bool& RefValue, const TCHAR* Help);
+	CORE_API FAutoConsoleVariableRef(uint32 Flags, const TCHAR* Name, FString& RefValue, const TCHAR* Help);
+	CORE_API FAutoConsoleVariableRef(uint32 Flags, const TCHAR* Name, FName& RefValue, const TCHAR* Help);
+	CORE_API FAutoConsoleVariableRef(uint32 Flags, const TCHAR* Name, int32& RefValue, const TCHAR* Help, const FConsoleVariableDelegate& Callback);
+	CORE_API FAutoConsoleVariableRef(uint32 Flags, const TCHAR* Name, float& RefValue, const TCHAR* Help, const FConsoleVariableDelegate& Callback);
+	CORE_API FAutoConsoleVariableRef(uint32 Flags, const TCHAR* Name, bool& RefValue, const TCHAR* Help, const FConsoleVariableDelegate& Callback);
+	CORE_API FAutoConsoleVariableRef(uint32 Flags, const TCHAR* Name, FString& RefValue, const TCHAR* Help, const FConsoleVariableDelegate& Callback);
+
+	CORE_API FAutoConsoleVariableRef(const FCVarDeprecationArgs& DeprecationArgs, const TCHAR* Name, int32& RefValue, uint32 Flags);
+	CORE_API FAutoConsoleVariableRef(const FCVarDeprecationArgs& DeprecationArgs, const TCHAR* Name, float& RefValue, uint32 Flags);
+	CORE_API FAutoConsoleVariableRef(const FCVarDeprecationArgs& DeprecationArgs, const TCHAR* Name, bool& RefValue, uint32 Flags);
+	CORE_API FAutoConsoleVariableRef(const FCVarDeprecationArgs& DeprecationArgs, const TCHAR* Name, FString& RefValue, uint32 Flags);
+	CORE_API FAutoConsoleVariableRef(const FCVarDeprecationArgs& DeprecationArgs, const TCHAR* Name, FName& RefValue, uint32 Flags);
+	CORE_API FAutoConsoleVariableRef(const FCVarDeprecationArgs& DeprecationArgs, const TCHAR* Name, int32& RefValue, const FConsoleVariableDelegate& Callback, uint32 Flags);
+	CORE_API FAutoConsoleVariableRef(const FCVarDeprecationArgs& DeprecationArgs, const TCHAR* Name, float& RefValue, const FConsoleVariableDelegate& Callback, uint32 Flags);
+	CORE_API FAutoConsoleVariableRef(const FCVarDeprecationArgs& DeprecationArgs, const TCHAR* Name, bool& RefValue, const FConsoleVariableDelegate& Callback, uint32 Flags);
+	CORE_API FAutoConsoleVariableRef(const FCVarDeprecationArgs& DeprecationArgs, const TCHAR* Name, FString& RefValue, const FConsoleVariableDelegate& Callback, uint32 Flags);
+
 public:
 	/**
 	 * Create a reference to a int console variable
@@ -1676,32 +1727,42 @@ public:
 	 * @param Help must not be 0
 	 * @param Flags bitmask combined from EConsoleVariableFlags
 	 */
-	CORE_API FAutoConsoleVariableRef(const TCHAR* Name, int32& RefValue, const TCHAR* Help, uint32 Flags = ECVF_Default);
-	CORE_API FAutoConsoleVariableRef(FCVarDeprecationArgs DeprecationArgs, const TCHAR* Name, int32& RefValue, const TCHAR* Help, uint32 Flags = ECVF_Default);
+	FORCEINLINE FAutoConsoleVariableRef(const TCHAR* Name, int32& RefValue, const TCHAR* Help, uint32 Flags = ECVF_Default)
+		: FAutoConsoleVariableRef(Flags, Name, RefValue, UE_CVAR_HELP_TEXT_PARAM(Help)) {}
+	FORCEINLINE FAutoConsoleVariableRef(const FCVarDeprecationArgs& DeprecationArgs, const TCHAR* Name, int32& RefValue, const TCHAR* Help, uint32 Flags = ECVF_Default)
+		: FAutoConsoleVariableRef(DeprecationArgs, Name, RefValue, Flags) {}
+
 	/**
 	 * Create a reference to a float console variable
 	 * @param Name must not be 0
 	 * @param Help must not be 0
 	 * @param Flags bitmask combined from EConsoleVariableFlags
 	 */
-	CORE_API FAutoConsoleVariableRef(const TCHAR* Name, float& RefValue, const TCHAR* Help, uint32 Flags = ECVF_Default);
-	CORE_API FAutoConsoleVariableRef(FCVarDeprecationArgs DeprecationArgs, const TCHAR* Name, float& RefValue, const TCHAR* Help, uint32 Flags = ECVF_Default);
+	FORCEINLINE FAutoConsoleVariableRef(const TCHAR* Name, float& RefValue, const TCHAR* Help, uint32 Flags = ECVF_Default)
+		: FAutoConsoleVariableRef(Flags, Name, RefValue, UE_CVAR_HELP_TEXT_PARAM(Help)) {}
+	FORCEINLINE FAutoConsoleVariableRef(const FCVarDeprecationArgs& DeprecationArgs, const TCHAR* Name, float& RefValue, const TCHAR* Help, uint32 Flags = ECVF_Default)
+		: FAutoConsoleVariableRef(DeprecationArgs, Name, RefValue, Flags) {}
+
 	/**
 	 * Create a reference to a bool console variable
 	 * @param Name must not be 0
 	 * @param Help must not be 0
 	 * @param Flags bitmask combined from EConsoleVariableFlags
 	 */
-	CORE_API FAutoConsoleVariableRef(const TCHAR* Name, bool& RefValue, const TCHAR* Help, uint32 Flags = ECVF_Default);
-	CORE_API FAutoConsoleVariableRef(FCVarDeprecationArgs DeprecationArgs, const TCHAR* Name, bool& RefValue, const TCHAR* Help, uint32 Flags = ECVF_Default);
+	FORCEINLINE FAutoConsoleVariableRef(const TCHAR* Name, bool& RefValue, const TCHAR* Help, uint32 Flags = ECVF_Default)
+		: FAutoConsoleVariableRef(Flags, Name, RefValue, UE_CVAR_HELP_TEXT_PARAM(Help)) {}
+	FORCEINLINE FAutoConsoleVariableRef(const FCVarDeprecationArgs& DeprecationArgs, const TCHAR* Name, bool& RefValue, const TCHAR* Help, uint32 Flags = ECVF_Default)
+		: FAutoConsoleVariableRef(DeprecationArgs, Name, RefValue, Flags) {}
 	/**
 	 * Create a reference to a FString console variable
 	 * @param Name must not be 0
 	 * @param Help must not be 0
 	 * @param Flags bitmask combined from EConsoleVariableFlags
 	 */
-	CORE_API FAutoConsoleVariableRef(const TCHAR* Name, FString& RefValue, const TCHAR* Help, uint32 Flags = ECVF_Default);
-	CORE_API FAutoConsoleVariableRef(FCVarDeprecationArgs DeprecationArgs, const TCHAR* Name, FString& RefValue, const TCHAR* Help, uint32 Flags = ECVF_Default);
+	FORCEINLINE FAutoConsoleVariableRef(const TCHAR* Name, FString& RefValue, const TCHAR* Help, uint32 Flags = ECVF_Default)
+		: FAutoConsoleVariableRef(Flags, Name, RefValue, UE_CVAR_HELP_TEXT_PARAM(Help)) {}
+	FORCEINLINE FAutoConsoleVariableRef(const FCVarDeprecationArgs& DeprecationArgs, const TCHAR* Name, FString& RefValue, const TCHAR* Help, uint32 Flags = ECVF_Default)
+		: FAutoConsoleVariableRef(DeprecationArgs, Name, RefValue, Flags) {}
 
 	/**
 	 * Create a reference to a FName console variable
@@ -1709,8 +1770,10 @@ public:
 	 * @param Help must not be 0
 	 * @param Flags bitmask combined from EConsoleVariableFlags
 	 */
-	CORE_API FAutoConsoleVariableRef(const TCHAR* Name, FName& RefValue, const TCHAR* Help, uint32 Flags = ECVF_Default);
-	CORE_API FAutoConsoleVariableRef(FCVarDeprecationArgs DeprecationArgs, const TCHAR* Name, FName& RefValue, const TCHAR* Help, uint32 Flags = ECVF_Default);
+	FORCEINLINE FAutoConsoleVariableRef(const TCHAR* Name, FName& RefValue, const TCHAR* Help, uint32 Flags = ECVF_Default)
+		: FAutoConsoleVariableRef(Flags, Name, RefValue, UE_CVAR_HELP_TEXT_PARAM(Help)) {}
+	FORCEINLINE FAutoConsoleVariableRef(const FCVarDeprecationArgs& DeprecationArgs, const TCHAR* Name, FName& RefValue, const TCHAR* Help, uint32 Flags = ECVF_Default)
+		: FAutoConsoleVariableRef(DeprecationArgs, Name, RefValue, Flags) {}
 
 	/**
 	 * Create a reference to a int console variable
@@ -1719,8 +1782,10 @@ public:
 	 * @param Callback Delegate called when the variable changes. @see IConsoleVariable::SetOnChangedCallback
 	 * @param Flags bitmask combined from EConsoleVariableFlags
 	 */
-	CORE_API FAutoConsoleVariableRef(const TCHAR* Name, int32& RefValue, const TCHAR* Help, const FConsoleVariableDelegate& Callback, uint32 Flags = ECVF_Default);
-	CORE_API FAutoConsoleVariableRef(FCVarDeprecationArgs DeprecationArgs, const TCHAR* Name, int32& RefValue, const TCHAR* Help, const FConsoleVariableDelegate& Callback, uint32 Flags = ECVF_Default);
+	FORCEINLINE FAutoConsoleVariableRef(const TCHAR* Name, int32& RefValue, const TCHAR* Help, const FConsoleVariableDelegate& Callback, uint32 Flags = ECVF_Default)
+		: FAutoConsoleVariableRef(Flags, Name, RefValue, UE_CVAR_HELP_TEXT_PARAM(Help), Callback) {}
+	FORCEINLINE FAutoConsoleVariableRef(const FCVarDeprecationArgs& DeprecationArgs, const TCHAR* Name, int32& RefValue, const TCHAR* Help, const FConsoleVariableDelegate& Callback, uint32 Flags = ECVF_Default)
+		: FAutoConsoleVariableRef(DeprecationArgs, Name, RefValue, Callback, Flags) {}
 
 	/**
 	 * Create a reference to a float console variable
@@ -1729,8 +1794,10 @@ public:
 	 * @param Callback Delegate called when the variable changes. @see IConsoleVariable::SetOnChangedCallback
 	 * @param Flags bitmask combined from EConsoleVariableFlags
 	 */
-	CORE_API FAutoConsoleVariableRef(const TCHAR* Name, float& RefValue, const TCHAR* Help, const FConsoleVariableDelegate& Callback, uint32 Flags = ECVF_Default);
-	CORE_API FAutoConsoleVariableRef(FCVarDeprecationArgs DeprecationArgs, const TCHAR* Name, float& RefValue, const TCHAR* Help, const FConsoleVariableDelegate& Callback, uint32 Flags = ECVF_Default);
+	FORCEINLINE FAutoConsoleVariableRef(const TCHAR* Name, float& RefValue, const TCHAR* Help, const FConsoleVariableDelegate& Callback, uint32 Flags = ECVF_Default)
+		: FAutoConsoleVariableRef(Flags, Name, RefValue, UE_CVAR_HELP_TEXT_PARAM(Help), Callback) {}
+	FORCEINLINE FAutoConsoleVariableRef(const FCVarDeprecationArgs& DeprecationArgs, const TCHAR* Name, float& RefValue, const TCHAR* Help, const FConsoleVariableDelegate& Callback, uint32 Flags = ECVF_Default)
+		: FAutoConsoleVariableRef(DeprecationArgs, Name, RefValue, Callback, Flags) {}
 
 	/**
 	 * Create a reference to a bool console variable
@@ -1739,8 +1806,10 @@ public:
 	 * @param Callback Delegate called when the variable changes. @see IConsoleVariable::SetOnChangedCallback
 	 * @param Flags bitmask combined from EConsoleVariableFlags
 	 */
-	CORE_API FAutoConsoleVariableRef(const TCHAR* Name, bool& RefValue, const TCHAR* Help, const FConsoleVariableDelegate& Callback, uint32 Flags = ECVF_Default);
-	CORE_API FAutoConsoleVariableRef(FCVarDeprecationArgs DeprecationArgs, const TCHAR* Name, bool& RefValue, const TCHAR* Help, const FConsoleVariableDelegate& Callback, uint32 Flags = ECVF_Default);
+	FORCEINLINE FAutoConsoleVariableRef(const TCHAR* Name, bool& RefValue, const TCHAR* Help, const FConsoleVariableDelegate& Callback, uint32 Flags = ECVF_Default)
+		: FAutoConsoleVariableRef(Flags, Name, RefValue, UE_CVAR_HELP_TEXT_PARAM(Help), Callback) {}
+	FORCEINLINE FAutoConsoleVariableRef(const FCVarDeprecationArgs& DeprecationArgs, const TCHAR* Name, bool& RefValue, const TCHAR* Help, const FConsoleVariableDelegate& Callback, uint32 Flags = ECVF_Default)
+		: FAutoConsoleVariableRef(DeprecationArgs, Name, RefValue, Callback, Flags) {}
 
 	/**
 	 * Create a reference to a FString console variable
@@ -1749,8 +1818,10 @@ public:
 	 * @param Callback Delegate called when the variable changes. @see IConsoleVariable::SetOnChangedCallback
 	 * @param Flags bitmask combined from EConsoleVariableFlags
 	 */
-	CORE_API FAutoConsoleVariableRef(const TCHAR* Name, FString& RefValue, const TCHAR* Help, const FConsoleVariableDelegate& Callback, uint32 Flags = ECVF_Default);
-	CORE_API FAutoConsoleVariableRef(FCVarDeprecationArgs DeprecationArgs, const TCHAR* Name, FString& RefValue, const TCHAR* Help, const FConsoleVariableDelegate& Callback, uint32 Flags = ECVF_Default);
+	FORCEINLINE FAutoConsoleVariableRef(const TCHAR* Name, FString& RefValue, const TCHAR* Help, const FConsoleVariableDelegate& Callback, uint32 Flags = ECVF_Default)
+		: FAutoConsoleVariableRef(Flags, Name, RefValue, UE_CVAR_HELP_TEXT_PARAM(Help), Callback) {}
+	FORCEINLINE FAutoConsoleVariableRef(const FCVarDeprecationArgs& DeprecationArgs, const TCHAR* Name, FString& RefValue, const TCHAR* Help, const FConsoleVariableDelegate& Callback, uint32 Flags = ECVF_Default)
+		: FAutoConsoleVariableRef(DeprecationArgs, Name, RefValue, Callback, Flags) {}
 
 	virtual ~FAutoConsoleVariableRef()
 	{
@@ -1892,6 +1963,30 @@ private: // ----------------------------------------------------
 template <class T>
 class TAutoConsoleVariable : public FAutoConsoleObject
 {
+	FORCENOINLINE TAutoConsoleVariable(uint32 Flags, const TCHAR* Name, const T& DefaultValue, const TCHAR* Help)
+		: TAutoConsoleVariable(IConsoleManager::Get().RegisterConsoleVariable(Name, DefaultValue, Help, Flags))
+	{
+	}
+
+	CORE_API TAutoConsoleVariable(const FCVarDeprecationArgs& DeprecationArgs, const TCHAR* Name, const T& DefaultValue, uint32 Flags);
+
+	FORCENOINLINE TAutoConsoleVariable(uint32 Flags, const TCHAR* Name, const T& DefaultValue, const TCHAR* Help, const FConsoleVariableDelegate& Callback)
+		: TAutoConsoleVariable(Flags, Name, DefaultValue, Help)
+	{
+		AsVariable()->SetOnChangedCallback(Callback);
+	}
+
+	FORCENOINLINE TAutoConsoleVariable(const FCVarDeprecationArgs& DeprecationArgs, const TCHAR* Name, const T& DefaultValue, const FConsoleVariableDelegate& Callback, uint32 Flags)
+		: TAutoConsoleVariable(DeprecationArgs, Name, DefaultValue, Flags)
+	{
+		AsVariable()->SetOnChangedCallback(Callback);
+	}
+
+	FORCENOINLINE TAutoConsoleVariable(uint32 Flags, const TCHAR* Name, const T& DefaultValue, TFunction<T()> Getter, TFunction<void(const T&)> Setter, const TCHAR* Help)
+		: TAutoConsoleVariable(IConsoleManager::Get().RegisterDelegatedConsoleVariable(Name, DefaultValue, MoveTemp(Getter), MoveTemp(Setter), Help, Flags))
+	{
+	}
+
 public:
 
 	TAutoConsoleVariable(IConsoleVariable* ExistingObject)
@@ -1925,12 +2020,15 @@ public:
 	 * @param Help must not be 0
 	 * @param Flags bitmask combined from EConsoleVariableFlags
 	 */
-	FORCENOINLINE TAutoConsoleVariable(const TCHAR* Name, const T& DefaultValue, const TCHAR* Help, uint32 Flags = ECVF_Default)
-		: TAutoConsoleVariable(IConsoleManager::Get().RegisterConsoleVariable(Name, DefaultValue, Help, Flags))
+	FORCEINLINE TAutoConsoleVariable(const TCHAR* Name, const T& DefaultValue, const TCHAR* Help, uint32 Flags = ECVF_Default)
+		: TAutoConsoleVariable(Flags, Name, DefaultValue, UE_CVAR_HELP_TEXT_PARAM(Help))
 	{
 	}
 
-	CORE_API TAutoConsoleVariable(FCVarDeprecationArgs DeprecationArgs, const TCHAR* Name, const T& DefaultValue, const TCHAR* Help, uint32 Flags = ECVF_Default);
+	FORCEINLINE TAutoConsoleVariable(const FCVarDeprecationArgs& DeprecationArgs, const TCHAR* Name, const T& DefaultValue, const TCHAR* Help, uint32 Flags = ECVF_Default)
+		: TAutoConsoleVariable(DeprecationArgs, Name, DefaultValue, Flags)
+	{
+	}
 	
 	/**
 	 * Create a float, int or string console variable
@@ -1939,16 +2037,14 @@ public:
 	 * @param Callback Delegate called when the variable changes. @see IConsoleVariable::SetOnChangedCallback
 	 * @param Flags bitmask combined from EConsoleVariableFlags
 	 */
-	FORCENOINLINE TAutoConsoleVariable(const TCHAR* Name, const T& DefaultValue, const TCHAR* Help, const FConsoleVariableDelegate& Callback, uint32 Flags = ECVF_Default)
-		: TAutoConsoleVariable(Name, DefaultValue, Help, Flags)
+	FORCEINLINE TAutoConsoleVariable(const TCHAR* Name, const T& DefaultValue, const TCHAR* Help, const FConsoleVariableDelegate& Callback, uint32 Flags = ECVF_Default)
+		: TAutoConsoleVariable(Flags, Name, DefaultValue, UE_CVAR_HELP_TEXT_PARAM(Help), Callback)
 	{
-		AsVariable()->SetOnChangedCallback(Callback);
 	}
 
-	FORCENOINLINE TAutoConsoleVariable(FCVarDeprecationArgs DeprecationArgs, const TCHAR* Name, const T& DefaultValue, const TCHAR* Help, const FConsoleVariableDelegate& Callback, uint32 Flags = ECVF_Default)
-		: TAutoConsoleVariable(DeprecationArgs,Name, DefaultValue, Help, Flags)
+	FORCEINLINE TAutoConsoleVariable(const FCVarDeprecationArgs& DeprecationArgs, const TCHAR* Name, const T& DefaultValue, const TCHAR* Help, const FConsoleVariableDelegate& Callback, uint32 Flags = ECVF_Default)
+		: TAutoConsoleVariable(DeprecationArgs, Name, DefaultValue, Callback, Flags)
 	{
-		AsVariable()->SetOnChangedCallback(Callback);
 	}
 
 	/**
@@ -1963,12 +2059,8 @@ public:
 	 * @param Help must not be 0
 	 * @param Flags bitmask combined from EConsoleVariableFlags
 	 */
-	FORCENOINLINE TAutoConsoleVariable(
-		const TCHAR* Name, const T& DefaultValue,
-		TFunction<T()> Getter, TFunction<void(const T&)> Setter,
-		const TCHAR* Help, uint32 Flags = ECVF_Default)
-		: TAutoConsoleVariable(IConsoleManager::Get().RegisterDelegatedConsoleVariable(
-			Name, DefaultValue, MoveTemp(Getter), MoveTemp(Setter), Help, Flags))
+	FORCEINLINE TAutoConsoleVariable(const TCHAR* Name, const T& DefaultValue, TFunction<T()> Getter, TFunction<void(const T&)> Setter, const TCHAR* Help, uint32 Flags = ECVF_Default)
+		: TAutoConsoleVariable(Flags, Name, DefaultValue, MoveTemp(Getter), MoveTemp(Setter), UE_CVAR_HELP_TEXT_PARAM(Help))
 	{
 	}
 
@@ -2090,20 +2182,22 @@ public:
 	virtual FString		GetString()		const override	{ return GetImpl<FString>(); }
 	virtual bool		GetBool()		const override { return GetImpl<bool>(); }
 
+#if UE_WITH_CVAR_HELP_TEXT
 	virtual const TCHAR* GetHelp() const override
 	{
 		return TEXT("NO_CVARS, no help");
 	}
 
-	virtual FText GetDetailedHelp() const override
+	virtual void SetHelp(const TCHAR* InHelp) override
 	{
-		return INVTEXT("NO_CVARS, no help");
+		check(false);
 	}
 
-	virtual void SetHelp(const TCHAR* InHelp) override
+	virtual FText GetDetailedHelp() const override
 	{
-		check(false);
+		return INVTEXT("NO_CVARS, no help");
 	}
+#endif
 
 	virtual void Release() override
 	{
@@ -2178,6 +2272,12 @@ private:
  */
 class FAutoConsoleCommand : private FAutoConsoleObject
 {
+	CORE_API FAutoConsoleCommand(uint32 Flags, const TCHAR* Name, const TCHAR* Help, const FConsoleCommandDelegate& Command);
+	CORE_API FAutoConsoleCommand(uint32 Flags, const TCHAR* Name, const TCHAR* Help, const FConsoleCommandWithArgsDelegate& Command);
+	CORE_API FAutoConsoleCommand(uint32 Flags, const TCHAR* Name, const TCHAR* Help, const FConsoleCommandWithOutputDeviceDelegate& Command);
+	CORE_API FAutoConsoleCommand(uint32 Flags, const TCHAR* Name, const TCHAR* Help, const FConsoleCommandWithArgsAndOutputDeviceDelegate& Command);
+	CORE_API FAutoConsoleCommand(uint32 Flags, const TCHAR* Name, const TCHAR* Help, const FConsoleCommandWithWorldDelegate& Command);
+	CORE_API FAutoConsoleCommand(uint32 Flags, const TCHAR* Name, const TCHAR* Help, const FConsoleCommandWithWorldArgsAndOutputDeviceDelegate& Command);
 public:
 	/**
 	 * Register a console command that takes no arguments
@@ -2187,7 +2287,8 @@ public:
 	 * @param	Command		The user function to call when this command is executed
 	 * @param	Flags		Optional flags bitmask
 	 */
-	CORE_API FAutoConsoleCommand(const TCHAR* Name, const TCHAR* Help, const FConsoleCommandDelegate& Command, uint32 Flags = ECVF_Default);
+	FORCEINLINE FAutoConsoleCommand(const TCHAR* Name, const TCHAR* Help, const FConsoleCommandDelegate& Command, uint32 Flags = ECVF_Default)
+		: FAutoConsoleCommand(Flags, Name, UE_CVAR_HELP_TEXT_PARAM(Help), Command) {}
 
 	/**
 	 * Register a console command that takes arguments
@@ -2197,7 +2298,8 @@ public:
 	 * @param	Command		The user function to call when this command is executed
 	 * @param	Flags		Optional flags bitmask
 	 */
-	CORE_API FAutoConsoleCommand(const TCHAR* Name, const TCHAR* Help, const FConsoleCommandWithArgsDelegate& Command, uint32 Flags = ECVF_Default);
+	FORCEINLINE FAutoConsoleCommand(const TCHAR* Name, const TCHAR* Help, const FConsoleCommandWithArgsDelegate& Command, uint32 Flags = ECVF_Default)
+		: FAutoConsoleCommand(Flags, Name, UE_CVAR_HELP_TEXT_PARAM(Help), Command) {}
 
 	/**
 	* Register a console command that takes an output device
@@ -2207,7 +2309,8 @@ public:
 	* @param	Command		The user function to call when this command is executed
 	* @param	Flags		Optional flags bitmask
 	*/
-	CORE_API FAutoConsoleCommand(const TCHAR* Name, const TCHAR* Help, const FConsoleCommandWithOutputDeviceDelegate& Command, uint32 Flags = ECVF_Default);
+	FORCEINLINE FAutoConsoleCommand(const TCHAR* Name, const TCHAR* Help, const FConsoleCommandWithOutputDeviceDelegate& Command, uint32 Flags = ECVF_Default)
+		: FAutoConsoleCommand(Flags, Name, UE_CVAR_HELP_TEXT_PARAM(Help), Command) {}
 
     /**
     * Register a console command that takes an output device
@@ -2217,7 +2320,8 @@ public:
     * @param	Command		The user function to call when this command is executed
     * @param	Flags		Optional flags bitmask
     */
-    CORE_API FAutoConsoleCommand(const TCHAR* Name, const TCHAR* Help, const FConsoleCommandWithArgsAndOutputDeviceDelegate& Command, uint32 Flags = ECVF_Default);
+	FORCEINLINE FAutoConsoleCommand(const TCHAR* Name, const TCHAR* Help, const FConsoleCommandWithArgsAndOutputDeviceDelegate& Command, uint32 Flags = ECVF_Default)
+		: FAutoConsoleCommand(Flags, Name, UE_CVAR_HELP_TEXT_PARAM(Help), Command) {}
 
 	/**
 	* Register a console command that takes a world argument
@@ -2227,7 +2331,8 @@ public:
 	* @param	Command		The user function to call when this command is executed
 	* @param	Flags		Optional flags bitmask
 	*/
-	CORE_API FAutoConsoleCommand(const TCHAR* Name, const TCHAR* Help, const FConsoleCommandWithWorldDelegate& Command, uint32 Flags = ECVF_Default);
+	FORCEINLINE FAutoConsoleCommand(const TCHAR* Name, const TCHAR* Help, const FConsoleCommandWithWorldDelegate& Command, uint32 Flags = ECVF_Default)
+		: FAutoConsoleCommand(Flags, Name, UE_CVAR_HELP_TEXT_PARAM(Help), Command) {}
 
 	/**
 	* Register a console command that takes arguments, a world argument and an output device
@@ -2237,7 +2342,10 @@ public:
 	* @param	Command		The user function to call when this command is executed
 	* @param	Flags		Optional flags bitmask
 	*/
-	CORE_API FAutoConsoleCommand(const TCHAR* Name, const TCHAR* Help, const FConsoleCommandWithWorldArgsAndOutputDeviceDelegate& Command, uint32 Flags = ECVF_Default);
+	FORCEINLINE FAutoConsoleCommand(const TCHAR* Name, const TCHAR* Help, const FConsoleCommandWithWorldArgsAndOutputDeviceDelegate& Command, uint32 Flags = ECVF_Default)
+		: FAutoConsoleCommand(Flags, Name, UE_CVAR_HELP_TEXT_PARAM(Help), Command) {}
+
+	virtual ~FAutoConsoleCommand() {}
 };
 #else
 class FAutoConsoleCommand
@@ -2276,6 +2384,7 @@ public:
  */
 class FAutoConsoleCommandWithWorld : private FAutoConsoleObject
 {
+	CORE_API FAutoConsoleCommandWithWorld(uint32 Flags, const TCHAR* Name, const TCHAR* Help, const FConsoleCommandWithWorldDelegate& Command);
 public:
 	/**
 	 * Register a console command that takes a world argument
@@ -2285,7 +2394,10 @@ public:
 	 * @param	Command		The user function to call when this command is executed
 	 * @param	Flags		Optional flags bitmask
 	 */
-	CORE_API FAutoConsoleCommandWithWorld(const TCHAR* Name, const TCHAR* Help, const FConsoleCommandWithWorldDelegate& Command, uint32 Flags = ECVF_Default);
+	FORCEINLINE FAutoConsoleCommandWithWorld(const TCHAR* Name, const TCHAR* Help, const FConsoleCommandWithWorldDelegate& Command, uint32 Flags = ECVF_Default)
+		: FAutoConsoleCommandWithWorld(Flags, Name, UE_CVAR_HELP_TEXT_PARAM(Help), Command) {}
+
+	virtual ~FAutoConsoleCommandWithWorld() {}
 };
 
 /**
@@ -2293,6 +2405,7 @@ public:
  */
 class FAutoConsoleCommandWithWorldAndArgs : private FAutoConsoleObject
 {
+	CORE_API FAutoConsoleCommandWithWorldAndArgs(uint32 Flags, const TCHAR* Name, const TCHAR* Help, const FConsoleCommandWithWorldAndArgsDelegate& Command);
 public:	
 	/**
 	 * Register a console command that takes arguments and a world argument
@@ -2302,7 +2415,10 @@ public:
 	 * @param	Command		The user function to call when this command is executed
 	 * @param	Flags		Optional flags bitmask
 	 */
-	CORE_API FAutoConsoleCommandWithWorldAndArgs(const TCHAR* Name, const TCHAR* Help, const FConsoleCommandWithWorldAndArgsDelegate& Command, uint32 Flags = ECVF_Default);
+	FORCEINLINE FAutoConsoleCommandWithWorldAndArgs(const TCHAR* Name, const TCHAR* Help, const FConsoleCommandWithWorldAndArgsDelegate& Command, uint32 Flags = ECVF_Default)
+		: FAutoConsoleCommandWithWorldAndArgs(Flags, Name, UE_CVAR_HELP_TEXT_PARAM(Help), Command) {}
+
+	virtual ~FAutoConsoleCommandWithWorldAndArgs() {}
 };
 
 /**
@@ -2310,9 +2426,13 @@ public:
  */
 class FAutoConsoleCommandWithArgsAndOutputDevice : private FAutoConsoleObject
 {
+	CORE_API FAutoConsoleCommandWithArgsAndOutputDevice(uint32 Flags, const TCHAR* Name, const TCHAR* Help, const FConsoleCommandWithArgsAndOutputDeviceDelegate& Command);
 public:
 	
-	CORE_API FAutoConsoleCommandWithArgsAndOutputDevice(const TCHAR* Name, const TCHAR* Help, const FConsoleCommandWithArgsAndOutputDeviceDelegate& Command, uint32 Flags = ECVF_Default);
+	FORCEINLINE FAutoConsoleCommandWithArgsAndOutputDevice(const TCHAR* Name, const TCHAR* Help, const FConsoleCommandWithArgsAndOutputDeviceDelegate& Command, uint32 Flags = ECVF_Default)
+		: FAutoConsoleCommandWithArgsAndOutputDevice(Flags, Name, UE_CVAR_HELP_TEXT_PARAM(Help), Command) {}
+
+	virtual ~FAutoConsoleCommandWithArgsAndOutputDevice() {}
 };
 
 /**
@@ -2320,6 +2440,7 @@ public:
  */
 class FAutoConsoleCommandWithOutputDevice : private FAutoConsoleObject
 {
+	CORE_API FAutoConsoleCommandWithOutputDevice(uint32 Flags, const TCHAR* Name, const TCHAR* Help, const FConsoleCommandWithOutputDeviceDelegate& Command);
 public:
 	/**
 	 * Register a console command that takes an output device
@@ -2329,7 +2450,10 @@ public:
 	 * @param	Command		The user function to call when this command is executed
 	 * @param	Flags		Optional flags bitmask
 	 */
-	CORE_API FAutoConsoleCommandWithOutputDevice(const TCHAR* Name, const TCHAR* Help, const FConsoleCommandWithOutputDeviceDelegate& Command, uint32 Flags = ECVF_Default);
+	FORCEINLINE FAutoConsoleCommandWithOutputDevice(const TCHAR* Name, const TCHAR* Help, const FConsoleCommandWithOutputDeviceDelegate& Command, uint32 Flags = ECVF_Default)
+		: FAutoConsoleCommandWithOutputDevice(Flags, Name, UE_CVAR_HELP_TEXT_PARAM(Help), Command) {}
+
+	virtual ~FAutoConsoleCommandWithOutputDevice() {}
 };
 
 /**
@@ -2337,6 +2461,7 @@ public:
  */
 class FAutoConsoleCommandWithWorldArgsAndOutputDevice : private FAutoConsoleObject
 {
+	CORE_API FAutoConsoleCommandWithWorldArgsAndOutputDevice(uint32 Flags, const TCHAR* Name, const TCHAR* Help, const FConsoleCommandWithWorldArgsAndOutputDeviceDelegate& Command);
 public:
 	/**
 	 * Register a console command that takes an output device
@@ -2346,7 +2471,10 @@ public:
 	 * @param	Command		The user function to call when this command is executed
 	 * @param	Flags		Optional flags bitmask
 	 */
-	CORE_API FAutoConsoleCommandWithWorldArgsAndOutputDevice(const TCHAR* Name, const TCHAR* Help, const FConsoleCommandWithWorldArgsAndOutputDeviceDelegate& Command, uint32 Flags = ECVF_Default);
+	FORCEINLINE FAutoConsoleCommandWithWorldArgsAndOutputDevice(const TCHAR* Name, const TCHAR* Help, const FConsoleCommandWithWorldArgsAndOutputDeviceDelegate& Command, uint32 Flags = ECVF_Default)
+		: FAutoConsoleCommandWithWorldArgsAndOutputDevice(Flags, Name, UE_CVAR_HELP_TEXT_PARAM(Help), Command) {}
+
+	virtual ~FAutoConsoleCommandWithWorldArgsAndOutputDevice() {}
 };
 
 #else
@@ -2484,8 +2612,10 @@ public:
 template<typename T>
 class TAutoRangedConsoleVariable : public TAutoConsoleVariable<T>
 {
+	CORE_API TAutoRangedConsoleVariable(IConsoleVariable* Parent, const TCHAR* Name, T DefaultValue, T MinValue, T MaxValue, const TCHAR* Help, uint32 Flags);
 public:
-	CORE_API TAutoRangedConsoleVariable(const TCHAR* Name, T DefaultValue, T MinValue, T MaxValue, const TCHAR* Help, uint32 Flags = ECVF_Default, IConsoleVariable* Parent = nullptr);
+	FORCEINLINE TAutoRangedConsoleVariable(const TCHAR* Name, T DefaultValue, T MinValue, T MaxValue, const TCHAR* Help, uint32 Flags = ECVF_Default, IConsoleVariable* Parent = nullptr)
+		: TAutoRangedConsoleVariable(Parent, Name, Forward(DefaultValue), Forward(MinValue), Forward(MaxValue), UE_CVAR_HELP_TEXT_PARAM(Help), Flags) {}
 };
 
 
