Cleanup task headers to organize and improve comments:
Async/Fundamental/Task.h is for LowLevelTasks (this could be renamed in a future change)
Tasks/Task.h is the public API for UE::Tasks
Tasks/TaskPrivate.h is for helpers that should be considered private

#Synced-CL 51908825

diff --git a/Engine/Source/Runtime/Core/Public/Async/Fundamental/Task.h b/Engine/Source/Runtime/Core/Public/Async/Fundamental/Task.h
--- a/Engine/Source/Runtime/Core/Public/Async/Fundamental/Task.h
+++ b/Engine/Source/Runtime/Core/Public/Async/Fundamental/Task.h
@@ -12,20 +12,24 @@
 
 #define LOWLEVEL_TASK_SIZE PLATFORM_CACHE_LINE_SIZE
 
+/** This namespace and header is used for the internals of the task system, engine code should include Tasks/Task.h instead */
 namespace LowLevelTasks
 {
 	DECLARE_LOG_CATEGORY_EXTERN(LowLevelTasks, Log, All);
 
+	/** Defines the order that tasks are processed in, starting with High and ending with BackgroundLow */
 	enum class ETaskPriority : int8
 	{
 		High,
 		Normal,
 		Default = Normal,
 		ForegroundCount,
+
 		BackgroundHigh = ForegroundCount,
 		BackgroundNormal,
 		BackgroundLow,
 		Count,
+
 		Inherit, //Inherit the TaskPriority from the launching Task or the Default Priority if not launched from a Task.
 	};
 
@@ -82,6 +86,7 @@ namespace LowLevelTasks
 		return false;
 	}
 
+	/** Flags used to change the behavior of TryCancel */
 	enum class ECancellationFlags : int8
 	{
 		None					= 0 << 0,
@@ -91,6 +96,7 @@ namespace LowLevelTasks
 	};
 	ENUM_CLASS_FLAGS(ECancellationFlags)
 
+	/** Flags defining what operations are possible on a specific task */
 	enum class ETaskFlags : int8
 	{
 		AllowNothing		= 0 << 0,
@@ -101,7 +107,9 @@ namespace LowLevelTasks
 	};
 	ENUM_CLASS_FLAGS(ETaskFlags)
 
-	/*																	       
+	/**
+	 * The state of an individual task, following this state machine:
+	 *
 	 * (I)nitThread:                                                        STORE(I)----------------------CAS(C)----------------------    
 	 * (C)ancelingThread:                                                      --->|         Ready        |<-->|   CanceledAndReady   |   
 	 *                                                                              ----------------------      ----------------------    
@@ -146,10 +154,10 @@ namespace LowLevelTasks
 	};
 	ENUM_CLASS_FLAGS(ETaskState)
 
-	/*
-	* Generic implementation of a Deleter, it often comes up that one has to call a function to cleanup after a Task finished
-	* this can be done by capturing a TDeleter like so: [Deleter(LowLevelTasks::TDeleter<Type, &Type::DeleteFunction>(value))](){}
-	*/
+	/**
+	 * Generic implementation of a Deleter, it often comes up that one has to call a function to cleanup after a Task finished.
+	 * This can be done by capturing a TDeleter like so: [Deleter(LowLevelTasks::TDeleter<Type, &Type::DeleteFunction>(value))](){}
+	 */
 	template<typename Type, void (Type::*DeleteFunction)()>
 	class TDeleter
 	{
@@ -180,13 +188,14 @@ namespace LowLevelTasks
 		}
 	};
 
-	/*
-	* this class is just here to hide some variables away
-	* because we don't want to become too close friends with the FScheduler
-	*/
 	class FTask;
 	namespace Tasks_Impl
 	{
+
+	/**
+	 * This class is just here to hide some variables away
+	 * because we don't want to become too close friends with the FScheduler
+	 */
 	class FTaskBase
 	{
 		class FPackedDataAtomic;
@@ -304,9 +313,11 @@ namespace LowLevelTasks
 	};
 	}
 
-	/*
-	* minimal low level task interface
-	*/
+
+	/**
+	 * Low level task implementation, this provides the state machine interface.
+	 * The classes in UE::Tasks handle actual task execution.
+	 */
 	class FTask final : private Tasks_Impl::FTaskBase
 	{		
 		friend class FScheduler;
@@ -316,45 +327,28 @@ namespace LowLevelTasks
 
 	public:
 
-		/*
-		* means the task is completed and this taskhandle can be recycled
-		*/
+		/** If true, this means the task is completed and this taskhandle can be recycled */
 		inline bool IsCompleted(std::memory_order MemoryOrder = std::memory_order_seq_cst) const
 		{
 			ETaskState State = PackedData.load(MemoryOrder).GetState();
 			return EnumHasAnyFlags(State, ETaskState::CompletedFlag);
 		}
 		
-		/*
-		* means the task was canceled but might still need to be launched
-		*/
+		/** If true, this means the task was canceled but might still need to be launched */
 		inline bool WasCanceled() const
 		{
 			ETaskState State = PackedData.load(std::memory_order_relaxed).GetState();
 			return EnumHasAnyFlags(State, ETaskState::CanceledFlag);
 		}
 
-		/*
-		* means the task was expedited or that it already completed
-		*/
+		/** If true, this means the task was expedited or that it already completed */
 		inline bool WasExpedited() const
 		{
 			ETaskState State = PackedData.load(std::memory_order_acquire).GetState();
 			return EnumHasAnyFlags(State, ETaskState::ExpeditedFlag | ETaskState::CompletedFlag);
 		}
 
-	private:
-		//Scheduler internal interface to speed things up
-		inline bool WasCanceledOrIsExpediting() const
-		{
-			ETaskState State = PackedData.load(std::memory_order_relaxed).GetState();
-			return EnumHasAnyFlags(State, ETaskState::CanceledFlag | ETaskState::RunningFlag);
-		}
-
-	public:
-		/*
-		* means the task is ready to be launched but might already been canceled 
-		*/
+		/** If true, this means the task is ready to be launched but might have already been canceled */
 		inline bool IsReady() const 
 		{
 			ETaskState State = PackedData.load(std::memory_order_relaxed).GetState();
@@ -362,9 +356,7 @@ namespace LowLevelTasks
 		}
 
 #if PLATFORM_DESKTOP || !IS_MONOLITHIC
-		/*
-		* get the currently active task if any
-		*/
+		/** Called by internal code to get the low level task executing on the current thread, may return null */
 		CORE_API static const FTask* GetActiveTask();
 #else
 		UE_FORCEINLINE_HINT static const FTask* GetActiveTask()
@@ -373,33 +365,32 @@ namespace LowLevelTasks
 		}
 #endif
 
-		/*
-		* try to cancel the task if it has not been launched yet.
-		*/
+		/** Try to cancel the task if it has not been launched yet */
 		inline bool TryCancel(ECancellationFlags CancellationFlags = ECancellationFlags::DefaultFlags);
 
-		/*
-		* try to revive a canceled task (as in reverting the cancellation as if it never happened).
-		* if it had been canceled and the scheduler has not run it yet it succeeds.
-		*/
+		/** 
+		 * Try to revive a canceled task (as in reverting the cancellation as if it never happened).
+		 * If it had been canceled and the scheduler has not run it yet this will succeed and return true.
+		 */
 		inline bool TryRevive();
 
-		/*
-		* try to expedite the task if succeded it will run immediately 
-		* but it will not set the completed state until the scheduler has executed it, because the scheduler still holds a reference.
-		* to check for completion in the context of expediting use WasExpedited. The TaskHandle canot be reused until IsCompleted returns true.
-		* @param Continuation: optional Continuation that needs to be executed or scheduled by the caller (can only be non null if the operation returned true)
-		*/
+		/**
+		 * Try to expedite the task, if possible it will run immediately 
+		 * but it will not set the completed state until the scheduler has executed it, because the scheduler still holds a reference.
+		 * To check for completion in the context of expediting use WasExpedited. The TaskHandle cannot be reused until IsCompleted returns true.
+		 * @param Continuation: optional Continuation that needs to be executed or scheduled by the caller (can only be non null if the operation returned true)
+		 */
 		inline bool TryExpedite();
 		inline bool TryExpedite(FTask*& Continuation);
 
-		/*
-		* try to execute the task if it has not been launched yet the task will execute immediately.
-		* @param Continuation: optional Continuation that needs to be executed or scheduled by the caller (can only be non null if the operation returned true)
-		*/
+		/**
+		 * Try to execute the task, if it has not been launched yet the task will execute immediately.
+		 * @param Continuation: optional Continuation that needs to be executed or scheduled by the caller (can only be non null if the operation returned true)
+		 */
 		inline bool TryExecute();
 		inline bool TryExecute(FTask*& Continuation);
 
+		/** Initialize a new task and prepare for execution */
 		template<typename TRunnable>
 		AUTORTFM_INFER inline void Init(const TCHAR* InDebugName, ETaskPriority InPriority, TRunnable&& InRunnable, ETaskFlags Flags = ETaskFlags::DefaultFlags);
 
@@ -423,7 +414,6 @@ namespace LowLevelTasks
 		void* GetUserData() const { return UserData; }
 		void SetUserData(void* NewUserData) const { UserData = NewUserData; }
 
-	public:
 		FTask() = default;
 		inline ~FTask();
 
@@ -433,6 +423,12 @@ namespace LowLevelTasks
 			return ActiveTask && ActiveTask->IsBackgroundTask();
 		}
 
+		inline bool WasCanceledOrIsExpediting() const
+		{
+			ETaskState State = PackedData.load(std::memory_order_relaxed).GetState();
+			return EnumHasAnyFlags(State, ETaskState::CanceledFlag | ETaskState::RunningFlag);
+		}
+
 		inline bool TryPrepareLaunch();
 		//after calling this function the task can be considered dead
 		template<bool bIsExpeditingThread>
diff --git a/Engine/Source/Runtime/Core/Public/Tasks/Task.h b/Engine/Source/Runtime/Core/Public/Tasks/Task.h
--- a/Engine/Source/Runtime/Core/Public/Tasks/Task.h
+++ b/Engine/Source/Runtime/Core/Public/Tasks/Task.h
@@ -10,77 +10,63 @@
 #include "HAL/IConsoleManager.h"
 #include "CoreTypes.h"
 
+/** 
+ * This namespace provides utilities for creating and accessing the state of tasks that run on separate threads.
+ * You can do this by using one of the Launch functions to start a task and return a handle for later use. For example:
+ *
+ *		UE::Tasks::TTask<FReturnValue> TaskHandle = UE::Tasks::Launch(TEXT("TaskDebugName"), []() {
+ *			DoSomethingSlow();
+ *			return FReturnValue();
+ *		});
+ * 
+ * In this example, the task will run on any worker thread and return a result payload that can be accessed
+ * on the calling thread when the task has completed. Tasks can be launched with a list of prerequisites,
+ * which will delay execution of the task until the prerequisites have completed. UE::Tasks::FTaskEvent can 
+ * be used as a prerequisite to support manual triggering of task execution.
+ * 
+ * Read the full documentation for an explanation of the different parameters to the launch function and
+ * the different ways to use the task handle after it has been launched.
+ */
 namespace UE::Tasks
 {
-	template<typename ResultType>
-	class TTask;
-
-	template<typename TaskCollectionType>
-	bool Wait(const TaskCollectionType& Tasks, FTimespan InTimeout = FTimespan::MaxValue());
-
-	template<typename TaskType> void AddNested(const TaskType& Nested);
 
 	namespace Private
 	{
-		template<typename TaskCollectionType>
-		TArray<TaskTrace::FId> GetTraceIds(const TaskCollectionType& Tasks);
-
-		// a common part of the generic `TTask<ResultType>` and its `TTask<void>` specialisation
+		/** Base implementation used by TTask to provide handle functionality */
 		class FTaskHandle
 		{
-			// friends to get access to `Pimpl`
-			friend FTaskBase;
-
-			template<int32 Index, typename ArrayType, typename FirstTaskType, typename... OtherTasksTypes>
-			friend void PrerequisitesUnpacker(ArrayType& Array, FirstTaskType& FirstTask, OtherTasksTypes&... OtherTasks);
-
-			template<int32 Index, typename ArrayType, typename TaskType>
-			friend void PrerequisitesUnpacker(ArrayType& Array, TaskType& FirstTask);
-
-			template<typename TaskCollectionType>
-			friend bool TryRetractAndExecute(const TaskCollectionType& Tasks, FTimeout Timeout);
-
-			template<typename TaskCollectionType>
-			friend TArray<TaskTrace::FId> GetTraceIds(const TaskCollectionType& Tasks);
-
-			template<typename TaskCollectionType>
-			friend bool UE::Tasks::Wait(const TaskCollectionType& Tasks, FTimespan InTimeout);
-
-			template<typename TaskType>
-			friend void UE::Tasks::AddNested(const TaskType& Nested);
-
-		protected:
-			explicit FTaskHandle(FTaskBase* Other)
-				: Pimpl(Other, /*bAddRef = */false)
-			{}
-
 		public:
 			using FTaskHandleId = void;
 
 			FTaskHandle() = default;
 
+			/** True if the task has been correctly initialized and not yet destroyed */
 			bool IsValid() const
 			{
 				return Pimpl.IsValid();
 			}
 
-			// checks if task's execution is done
+			/** Checks if task's execution is done */
 			bool IsCompleted() const
 			{
 				return !IsValid() || Pimpl->IsCompleted();
 			}
 
-			// waits for task's completion with timeout. Tries to retract the task and execute it in-place, if failed - blocks until the task 
-			// is completed by another thread. If timeout is zero, tries to retract the task and returns immedially after that.
-			// @return true if the task is completed
+			/**
+			 * Waits for task's completion with timeout. Tries to retract the task and execute it in-place, if failed - blocks until the task 
+			 * is completed by another thread. If timeout is zero, tries to retract the task and returns immedially after that.
+			 * @return true if the task is completed
+			 */
 			bool Wait(FTimespan Timeout) const
 			{
 				return !IsValid() || Pimpl->Wait(FTimeout{ Timeout });
 			}
 
-			// waits for task's completion without timeout. Tries to retract the task and execute it in-place, if failed - blocks until the task 
-			// is completed by another thread.
-			// @return true if the task is completed
+			/**
+			 * Waits for task's completion without timeout. Tries to retract the task and execute it in-place, if failed - blocks until the task 
+			 * is completed by another thread.
+			 * @return true if the task is completed
+			 */
 			bool Wait() const
 			{
 				if (IsValid())
@@ -93,28 +79,32 @@ namespace UE::Tasks
 				return true;
 			}
 
-			// try to finish the task on the current thread if it's not already started.
-			// @return true if the task is completed
+			/**
+			 * Try to finish the task on the current thread if it's not already started.
+			 * @return true if the task is completed
+			 */
 			bool TryRetractAndExecute()
 			{
 				Pimpl->TryRetractAndExecute(FTimeout::Never());
 				return IsCompleted();
 			}
 
-			// launches a task for asynchronous execution
-			// @param DebugName - a unique name for task identification in debugger and profiler, is compiled out in test/shipping builds
-			// @param TaskBody - a functor that will be executed asynchronously
-			// @param Priority - task priority that affects when the task will be executed
-			// @param TaskFlags - task config options
-			// @return a trivially relocatable instance that can be used to wait for task completion or to obtain task execution result
+			/**
+			 * Launches a task for asynchronous execution
+			 * @param DebugName - a unique name for task identification in debugger and profiler, is compiled out in test/shipping builds
+			 * @param TaskBody - a functor that will be executed asynchronously
+			 * @param Priority - task priority that affects when the task will be executed
+			 * @param ExtendedPriority - priority options for integrating with other system like named threads
+	 		 * @param Flags - task config options
+			 * @return a trivially relocatable instance that can be used to wait for task completion or to obtain task execution result
+			 */
 			template<typename TaskBodyType>
 			void Launch(
 				const TCHAR* DebugName,
 				TaskBodyType&& TaskBody,
 				ETaskPriority Priority = ETaskPriority::Normal,
 				EExtendedTaskPriority ExtendedPriority = EExtendedTaskPriority::None,
-				ETaskFlags Flags = ETaskFlags::None
-			)
+				ETaskFlags Flags = ETaskFlags::None)
 			{
 				check(!IsValid());
 
@@ -125,14 +115,17 @@ namespace UE::Tasks
 				Task->TryLaunch(sizeof(*Task));
 			}
 
-			// launches a task for asynchronous execution, with prerequisites that must be completed before the task is scheduled
-			// @param DebugName - a unique name for task identification in debugger and profiler, is compiled out in test/shipping builds
-			// @param TaskBody - a functor that will be executed asynchronously
-			// @param Prerequisites - tasks or task events that must be completed before the task being launched can be scheduled, accepts any 
-			// iterable collection (.begin()/.end()), `Tasks::Prerequisites()` helper is recommended to create such collection on the fly
-			// @param Priority - task priority that affects when the task will be executed
-			// @param TaskFlags - task config options
-			// @return a trivially relocatable instance that can be used to wait for task completion or to obtain task execution result
+			/**
+			 * Launches a task for asynchronous execution, with prerequisites that must be completed before the task is scheduled
+			 * @param DebugName - a unique name for task identification in debugger and profiler, is compiled out in test/shipping builds
+			 * @param TaskBody - a functor that will be executed asynchronously
+			 * @param Prerequisites - tasks or task events that must be completed before the task being launched can be scheduled, accepts any 
+			 *		iterable collection (.begin()/.end()), `Tasks::Prerequisites()` helper is recommended to create such collection on the fly
+			 * @param Priority - task priority that affects when the task will be executed
+			 * @param ExtendedPriority - priority options for integrating with other system like named threads
+			 * @param Flags - task config options
+			 * @return a trivially relocatable instance that can be used to wait for task completion or to obtain task execution result
+			 */
 			template<typename TaskBodyType, typename PrerequisitesCollectionType>
 			void Launch(
 				const TCHAR* DebugName,
@@ -140,8 +133,7 @@ namespace UE::Tasks
 				PrerequisitesCollectionType&& Prerequisites,
 				ETaskPriority Priority = ETaskPriority::Normal,
 				EExtendedTaskPriority ExtendedPriority = EExtendedTaskPriority::None,
-				ETaskFlags Flags = ETaskFlags::None
-			)
+				ETaskFlags Flags = ETaskFlags::None)
 			{
 				check(!IsValid());
 
@@ -153,11 +145,24 @@ namespace UE::Tasks
 				Task->TryLaunch(sizeof(*Task));
 			}
 
+			/** Returns true if it is valid to Wait on this task */
 			bool IsAwaitable() const
 			{
 				return IsValid() && Pimpl->IsAwaitable();
 			}
 
+			/** Gets the priority this task was initialized with */
+			ETaskPriority GetPriority() const
+			{
+				return Pimpl->GetPriority();
+			}
+
+			/** Gets the extended priority this task was initialized with */
+			EExtendedTaskPriority GetExtendedPriority() const
+			{
+				return Pimpl->GetExtendedPriority();
+			}
+
 			bool operator==(const FTaskHandle& Other) const
 			{
 				return Pimpl == Other.Pimpl;
@@ -168,31 +173,51 @@ namespace UE::Tasks
 				return Pimpl != Other.Pimpl;
 			}
 			
-			ETaskPriority GetPriority() const
-			{
-				return Pimpl->GetPriority();
-			}
+		private:
+			// friends to get access to `Pimpl`
+			friend FTaskBase;
 
-			EExtendedTaskPriority GetExtendedPriority() const
-			{
-				return Pimpl->GetExtendedPriority();
-			}
-		
+			template<int32 Index, typename ArrayType, typename FirstTaskType, typename... OtherTasksTypes>
+			friend void PrerequisitesUnpacker(ArrayType& Array, FirstTaskType& FirstTask, OtherTasksTypes&... OtherTasks);
+
+			template<int32 Index, typename ArrayType, typename TaskType>
+			friend void PrerequisitesUnpacker(ArrayType& Array, TaskType& FirstTask);
+
+			template<typename TaskCollectionType>
+			friend bool TryRetractAndExecute(const TaskCollectionType& Tasks, FTimeout Timeout);
+
+			template<typename TaskCollectionType>
+			friend TArray<TaskTrace::FId> GetTraceIds(const TaskCollectionType& Tasks);
+
+			template<typename TaskCollectionType>
+			friend bool UE::Tasks::Wait(const TaskCollectionType& Tasks, FTimespan InTimeout);
+
+			template<typename TaskType>
+			friend void UE::Tasks::AddNested(const TaskType& Nested);
 
 		protected:
+			explicit FTaskHandle(FTaskBase* Other)
+				: Pimpl(Other, /*bAddRef = */false)
+			{}
+
 			TRefCountPtr<FTaskBase> Pimpl;
 		};
 	}
 
-	// a movable/copyable handle of `Private::TTaskWithResult*` with the API adopted for public usage.
-	// implements Pimpl idiom
+	/** Generic task type, can be either a real task or event */
+	using FTask = Private::FTaskHandle;
+
+	/** A templated asynchronous task that returns a result payload when complete */
 	template<typename ResultType>
 	class TTask : public Private::FTaskHandle
 	{
+		// a movable/copyable handle of `Private::TTaskWithResult*` with the API adopted for public usage.
+		// implements Pimpl idiom
+
 	public:
 		TTask() = default;
 
-		// waits until the task is completed and returns task's result
+		/** Waits until the task is completed and returns task's result */
 		ResultType& GetResult()
 		{
 			check(IsValid());
@@ -209,12 +234,14 @@ namespace UE::Tasks
 		{}
 	};
 
+	/** A task that does not return a result payload */
 	template<>
 	class TTask<void> : public Private::FTaskHandle
 	{
 	public:
 		TTask() = default;
 
+		/** Waits until task is completed, there is no result payload */
 		void GetResult()
 		{
 			check(IsValid()); // to be consistent with a generic `TTask<ResultType>::GetResult()`
@@ -230,8 +257,10 @@ namespace UE::Tasks
 		{}
 	};
 
-	// A synchronisation primitive, a recommended substitution of `FEvent` for signalling between tasks. If used as a task prerequisite or 
-	// a nested task, it doesn't block a worker thread. Optionally can use "busy waiting" - executing tasks while waiting.
+	/** 
+	 * A synchronisation primitive, a recommended substitution of `FEvent` for signalling between tasks. If used as a task prerequisite or 
+	 * a nested task, it doesn't block a worker thread. Optionally can use "busy waiting" - executing tasks while waiting.
+	 */
 	class FTaskEvent : public Private::FTaskHandle
 	{
 	public:
@@ -240,13 +269,14 @@ namespace UE::Tasks
 		{
 		}
 
-		// all prerequisites must be added before triggering the event
+		/** Add tasks to wait on, all prerequisites must be added before triggering the event */
 		template<typename PrerequisitesType>
 		void AddPrerequisites(const PrerequisitesType& Prerequisites)
 		{
 			Pimpl->AddPrerequisites(Prerequisites);
 		}
 
+		/** Activates the task event, which will complete if the prerequisites have completed */
 		void Trigger()
 		{
 			if (!IsCompleted()) // event can be triggered multiple times
@@ -256,20 +286,22 @@ namespace UE::Tasks
 		}
 	};
 
-	// launches a task for asynchronous execution
-	// @param DebugName - a unique name for task identification in debugger and profiler, is compiled out in test/shipping builds
-	// @param TaskBody - a functor that will be executed asynchronously
-	// @param Priority - task priority that affects when the task will be executed
-	// @param TaskFlags - task config options
-	// @return a trivially relocatable instance that can be used to wait for task completion or to obtain task execution result
+	/**
+	 * Launches a task for asynchronous execution
+	 * @param DebugName - a unique name for task identification in debugger and profiler, is compiled out in test/shipping builds
+	 * @param TaskBody - a functor that will be executed asynchronously
+	 * @param Priority - task priority that affects when the task will be executed
+	 * @param ExtendedPriority - priority options for integrating with other system like named threads
+	 * @param Flags - task config options
+	 * @return a trivially relocatable instance that can be used to wait for task completion or to obtain task execution result
+	 */
 	template<typename TaskBodyType>
 	TTask<TInvokeResult_T<TaskBodyType>> Launch(
 		const TCHAR* DebugName,
 		TaskBodyType&& TaskBody,
 		ETaskPriority Priority = ETaskPriority::Normal,
 		EExtendedTaskPriority ExtendedPriority = EExtendedTaskPriority::None,
-		ETaskFlags Flags = ETaskFlags::None
-	)
+		ETaskFlags Flags = ETaskFlags::None)
 	{
 		using FResult = TInvokeResult_T<TaskBodyType>;
 		TTask<FResult> Task;
@@ -277,14 +309,17 @@ namespace UE::Tasks
 		return Task;
 	}
 
-	// launches a task for asynchronous execution, with prerequisites that must be completed before the task is scheduled
-	// @param DebugName - a unique name for task identification in debugger and profiler, is compiled out in test/shipping builds
-	// @param TaskBody - a functor that will be executed asynchronously
-	// @param Prerequisites - tasks or task events that must be completed before the task being launched can be scheduled, accepts any 
-	// iterable collection (.begin()/.end()), `Tasks::Prerequisites()` helper is recommended to create such collection on the fly
-	// @param Priority - task priority that affects when the task will be executed
-	// @param TaskFlags - task config options
-	// @return a trivially relocatable instance that can be used to wait for task completion or to obtain task execution result
+	/**
+	 * Launches a task for asynchronous execution, with prerequisites that must be completed before the task is scheduled
+	 * @param DebugName - a unique name for task identification in debugger and profiler, is compiled out in test/shipping builds
+	 * @param TaskBody - a functor that will be executed asynchronously
+	 * @param Prerequisites - tasks or task events that must be completed before the task being launched can be scheduled, accepts any 
+	 *		iterable collection (.begin()/.end()), `Tasks::Prerequisites()` helper is recommended to create such collection on the fly
+	 * @param Priority - task priority that affects when the task will be executed
+	 * @param ExtendedPriority - priority options for integrating with other system like named threads
+	 * @param Flags - task config options
+	 * @return a trivially relocatable instance that can be used to wait for task completion or to obtain task execution result
+	 */
 	template<typename TaskBodyType, typename PrerequisitesCollectionType>
 	TTask<TInvokeResult_T<TaskBodyType>> Launch(
 		const TCHAR* DebugName,
@@ -292,8 +327,7 @@ namespace UE::Tasks
 		PrerequisitesCollectionType&& Prerequisites,
 		ETaskPriority Priority = ETaskPriority::Normal,
 		EExtendedTaskPriority ExtendedPriority = EExtendedTaskPriority::None,
-		ETaskFlags Flags = ETaskFlags::None
-	)
+		ETaskFlags Flags = ETaskFlags::None)
 	{
 		using FResult = TInvokeResult_T<TaskBodyType>;
 		TTask<FResult> Task;
@@ -301,37 +335,6 @@ namespace UE::Tasks
 		return Task;
 	}
 
-	namespace Private
-	{
-		template<typename TaskCollectionType>
-		TArray<TaskTrace::FId> GetTraceIds(const TaskCollectionType& Tasks)
-		{
-#if UE_TASK_TRACE_ENABLED
-			TArray<TaskTrace::FId> TasksIds;
-			TasksIds.Reserve(Tasks.Num());
-
-			for (auto& Task : Tasks)
-			{
-				if (Task.IsValid())
-				{
-					TasksIds.Add(Task.Pimpl->GetTraceId());
-				}
-			}
-
-			return TasksIds;
-#else
-			return {};
-#endif
-		}
-	}
-
-	inline void Wait(Private::FTaskHandle& Task)
-	{
-		Task.Wait();
-	}
-
-	using FTask = Private::FTaskHandle;
-
 	namespace Private
 	{
 		template<int32 Index, typename ArrayType, typename FirstTaskType, typename... OtherTasksTypes>
@@ -360,6 +363,7 @@ namespace UE::Tasks
 		}
 	}
 
+	/** Returns prerequisites for an array of tasks */
 	template<typename... TaskTypes, 
 		typename std::decay_t<decltype(std::declval<TTuple<TaskTypes...>>().template Get<0>())>::FTaskHandleId* = nullptr>
 	TStaticArray<Private::FTaskBase*, sizeof...(TaskTypes)> Prerequisites(TaskTypes&... Tasks)
@@ -375,8 +379,16 @@ namespace UE::Tasks
 		return Tasks;
 	}
 
-	// wait for multiple tasks, with timeout
-	// @param TaskCollectionType - an iterable collection of `TTask<T>`, e.g. `TArray<FTask>`
+	/** Utility function to wait on a specific task */
+	inline void Wait(FTask& Task)
+	{
+		Task.Wait();
+	}
+
+	/**
+	 * Wait for multiple tasks, with timeout
+	 * @param TaskCollectionType - an iterable collection of `TTask<T>`, e.g. `TArray<FTask>`
+	 */
 	template<typename TaskCollectionType>
 	bool Wait(const TaskCollectionType& Tasks, FTimespan InTimeout/* = FTimespan::MaxValue()*/)
 	{
@@ -396,9 +408,11 @@ namespace UE::Tasks
 	// "any task" support. these functions allocate excessively (per input task plus more).
 	// can be reduced to a single alloc if this is a perf issue
 	
-	// Blocks the current thread until any of the given tasks is completed.
-	// Is slightly more efficient than `Any().Wait()`.
-	// Returns the index of the first completed task, or `INDEX_NONE` on timeout
+	/**
+	 * Blocks the current thread until any of the given tasks is completed.
+	 * Is slightly more efficient than `Any().Wait()`.
+	 * Returns the index of the first completed task, or `INDEX_NONE` on timeout
+	 */
 	template<typename TaskCollectionType>
 	int32 WaitAny(const TaskCollectionType& Tasks, FTimespan Timeout = FTimespan::MaxValue())
 	{
@@ -449,7 +463,7 @@ namespace UE::Tasks
 		return INDEX_NONE;
 	}
 
-	// Returns a task that gets completed as soon as any of the given tasks gets completed
+	/** Returns a task that gets completed as soon as any of the given tasks gets completed */
 	template<typename TaskCollectionType>
 	FTask Any(const TaskCollectionType& Tasks)
 	{
@@ -501,9 +515,11 @@ namespace UE::Tasks
 		return Result;
 	}
 
-	// Adds the nested task to the task that is being currently executed by the current thread. A parent task is not flagged completed
-	// until all nested tasks are completed. It's similar to explicitly waiting for a sub-task at the end of its parent task, except explicit waiting
-	// blocks the worker executing the parent task until the sub-task is completed. With nested tasks, the worker won't be blocked.
+	/**
+	 * Adds the nested task to the task that is being currently executed by the current thread. A parent task is not flagged completed
+	 * until all nested tasks are completed. It's similar to explicitly waiting for a sub-task at the end of its parent task, except explicit waiting
+	 * blocks the worker executing the parent task until the sub-task is completed. With nested tasks, the worker won't be blocked.
+	 */
 	template<typename TaskType>
 	void AddNested(const TaskType& Nested)
 	{
@@ -512,10 +528,23 @@ namespace UE::Tasks
 		Parent->AddNested(*Nested.Pimpl);
 	}
 
-	// Console variable for configuring task priorities
-	// Example:
-	//		FTaskPriorityCVar CVar{ TEXT("CVarName"), TEXT("CVarHelp"), ETaskPriority::Normal, EExtendedTaskPriority::None };
-	//		Launch(UE_SOURCE_LOCATION, [] {}, CVar.GetTaskPriority(), CVar.GetExtendedTaskPriority()).Wait();
+	/** Creates and returns an already completed task holding a result value constructed from given args */
+	template<typename ResultType, typename... ArgTypes>
+	TTask<ResultType> MakeCompletedTask(ArgTypes&&... Args)
+	{
+		return Launch(
+			UE_SOURCE_LOCATION,
+			[&] { return ResultType(Forward<ArgTypes>(Args)...); },
+			ETaskPriority::Default, // doesn't matter
+			EExtendedTaskPriority::Inline);
+	}
+
+	/**
+	 * Console variable for configuring task priorities
+	 * Example:
+	 *		FTaskPriorityCVar CVar{ TEXT("CVarName"), TEXT("CVarHelp"), ETaskPriority::Normal, EExtendedTaskPriority::None };
+	 *		Launch(UE_SOURCE_LOCATION, [] {}, CVar.GetTaskPriority(), CVar.GetExtendedTaskPriority()).Wait();
+	 */
 	class FTaskPriorityCVar
 	{
 	public:
@@ -544,28 +573,19 @@ namespace UE::Tasks
 		EExtendedTaskPriority ExtendedPriority;
 	};
 
-	// creates and returns an already completed task holding a result value constructed from given args
-	template<typename ResultType, typename... ArgTypes>
-	TTask<ResultType> MakeCompletedTask(ArgTypes&&... Args)
-	{
-		return Launch(
-			UE_SOURCE_LOCATION,
-			[&] { return ResultType(Forward<ArgTypes>(Args)...); },
-			ETaskPriority::Default, // doesn't matter
-			EExtendedTaskPriority::Inline);
-	}
-
-	// support for canceling tasks mid-execution
-	// usage:
-	//		FCancellationToken Token;
-	//		Launch(UE_SOURCE_LOCATION, [&Token] { ... if (Token.IsCanceled()) return; ... });
-	//		Token.Cancel();
-	// * it's user's decision and responsibility to manage cancellation token lifetime, to check cancellation token, return early, 
-	//		to do any required cleanup, or to do nothing at all
-	// * no way to cancel a task to skip its execution completely
-	// * waiting for a canceled task is blocking until its prerequisites are completed and the task is executed and completed, 
-	//		basically same as for not canceled tasks except a canceled one can quit execution early
-	// * canceling a task doesn't affect its subsequents (unless they use the same cancellation token instance)
+	/**
+	 * Support for canceling tasks mid-execution
+	 * usage:
+	 *		FCancellationToken Token;
+	 *		Launch(UE_SOURCE_LOCATION, [&Token] { ... if (Token.IsCanceled()) return; ... });
+	 *		Token.Cancel();
+	 * it's user's decision and responsibility to manage cancellation token lifetime, to check cancellation token, return early, 
+	 *		to do any required cleanup, or to do nothing at all
+	 * no way to cancel a task to skip its execution completely
+	 * waiting for a canceled task is blocking until its prerequisites are completed and the task is executed and completed, 
+	 *		basically same as for not canceled tasks except a canceled one can quit execution early
+	 * canceling a task doesn't affect its subsequents (unless they use the same cancellation token instance)
+	 */
 	class FCancellationToken
 	{
 	public:
diff --git a/Engine/Source/Runtime/Core/Public/Tasks/TaskPrivate.h b/Engine/Source/Runtime/Core/Public/Tasks/TaskPrivate.h
--- a/Engine/Source/Runtime/Core/Public/Tasks/TaskPrivate.h
+++ b/Engine/Source/Runtime/Core/Public/Tasks/TaskPrivate.h
@@ -55,7 +55,7 @@ namespace UE::Tasks
 	using LowLevelTasks::ToString;
 	using LowLevelTasks::ToTaskPriority;
 
-	// special task priorities for tasks that are never sent to the scheduler
+	/** Options used to override general scheduler priority and integrate with other systems */
 	enum class EExtendedTaskPriority : int8
 	{
 		None,
@@ -91,18 +91,24 @@ namespace UE::Tasks
 		DoNotRunInsideBusyWait // do not pick this task for busy-waiting
 	};
 
-	namespace Private
-	{
-		CORE_API void TranslatePriority(ENamedThreads::Type ThreadType, ETaskPriority& OutPriority, EExtendedTaskPriority& OutExtendedPriority);
-		CORE_API ENamedThreads::Type TranslatePriority(ETaskPriority Priority, EExtendedTaskPriority ExtendedPriority);
-	}
-
 	class FPipe;
 
+	template<typename ResultType>
+	class TTask;
+
+	template<typename TaskCollectionType>
+	bool Wait(const TaskCollectionType& Tasks, FTimespan InTimeout = FTimespan::MaxValue());
+
+	template<typename TaskType> void AddNested(const TaskType& Nested);
+
+	/** Namespace for internals of task system, do not use these from general engine code */
 	namespace Private
 	{
 		class FTaskBase;
 
+		CORE_API void TranslatePriority(ENamedThreads::Type ThreadType, ETaskPriority& OutPriority, EExtendedTaskPriority& OutExtendedPriority);
+		CORE_API ENamedThreads::Type TranslatePriority(ETaskPriority Priority, EExtendedTaskPriority ExtendedPriority);
+
 		// returns the task (if any) that is being executed by the current thread
 		CORE_API FTaskBase* GetCurrentTask();
 		// sets the current task and returns the previous current task
@@ -1023,5 +1029,26 @@ protected:
 
 			return bResult;
 		}
+		
+		template<typename TaskCollectionType>
+		TArray<TaskTrace::FId> GetTraceIds(const TaskCollectionType& Tasks)
+		{
+#if UE_TASK_TRACE_ENABLED
+			TArray<TaskTrace::FId> TasksIds;
+			TasksIds.Reserve(Tasks.Num());
+
+			for (auto& Task : Tasks)
+			{
+				if (Task.IsValid())
+				{
+					TasksIds.Add(Task.Pimpl->GetTraceId());
+				}
+			}
+
+			return TasksIds;
+#else
+			return {};
+#endif
+		}
 	}
 }
