[CrashReportClient] Solve reading config using physical file reader

In some cases it is not possible to access the ini file (e.g. when bAllowNonUFSIniWhenCooked is set).

#Synced-CL 51621904

diff --git a/Engine/Source/Runtime/Core/Private/GenericPlatform/GenericPlatformCrashContext.cpp b/Engine/Source/Runtime/Core/Private/GenericPlatform/GenericPlatformCrashContext.cpp
--- a/Engine/Source/Runtime/Core/Private/GenericPlatform/GenericPlatformCrashContext.cpp
+++ b/Engine/Source/Runtime/Core/Private/GenericPlatform/GenericPlatformCrashContext.cpp
@@ -1776,11 +1776,15 @@ void FGenericCrashContext::CopyPlatformSpecificFiles(const TCHAR* OutputDirector
 {
 	// If present, include the crash report config file to pass config values to the CRC
 	const TCHAR* CrashConfigSrcPath = GetCrashConfigFilePath();
-	if (IFileManager::Get().FileExists(CrashConfigSrcPath))
+
+	// Use the physical file system directly rather than IFileManager, to avoid 
+	// being denied in scenarios where bAllowNonUFSIniWhenCooked=false
+	IPlatformFile* PhysicalFile = FPlatformFileManager::Get().FindPlatformFile(IPlatformFile::GetPhysicalTypeName());
+	if (PhysicalFile != nullptr && PhysicalFile->FileExists(CrashConfigSrcPath))
 	{
 		FString CrashConfigFilename = FPaths::GetCleanFilename(CrashConfigSrcPath);
 		const FString CrashConfigDstAbsolute = FPaths::Combine(OutputDirectory, *CrashConfigFilename);
-		IFileManager::Get().Copy(*CrashConfigDstAbsolute, CrashConfigSrcPath);	// best effort, so don't care about result: couldn't copy -> tough, no config
+		PhysicalFile->CopyFile(*CrashConfigDstAbsolute, CrashConfigSrcPath);	// best effort, so don't care about result: couldn't copy -> tough, no config
 	}
 
 	for (const FString& AbsoluteSource : NCached::FilePaths)
