[Insights] Emit diagnostic log message for user actions after trace panic.

Adding some diagnostic messages if user tries to start/stop traces after a panic, and also print a message in status output.

#Synced-CL 51621904

diff --git a/Engine/Source/Runtime/Core/Private/ProfilingDebugging/TraceAuxiliary.cpp b/Engine/Source/Runtime/Core/Private/ProfilingDebugging/TraceAuxiliary.cpp
--- a/Engine/Source/Runtime/Core/Private/ProfilingDebugging/TraceAuxiliary.cpp
+++ b/Engine/Source/Runtime/Core/Private/ProfilingDebugging/TraceAuxiliary.cpp
@@ -162,7 +162,7 @@ public:
 	FString GetSecureHost() const;
 	uint16 GetSecurePort() const;
 #endif
-	bool Stop();
+	bool Stop(const FTraceAuxiliary::FLogCategoryAlias& LogCategory);
 	void OnClosed();
 	void FreezeReadOnlyChannels();
 	void ResumeChannels();
@@ -181,6 +181,8 @@ public:
 	void ResetCommandlineChannels();
 	void UpdateStats();
 	void EmitStats(const UE::Trace::FStatistics& Stats);
+	void SetHasPanicked() { bHasPanicked = true; };
+	bool GetHasPanicked() { return bHasPanicked; };
 
 	// True if this is parent process with forking requested before forking.
 	bool IsParentProcessAndPreFork();
@@ -218,6 +220,7 @@ private:
 	bool bWorkerThreadStarted = false;
 	bool bTruncateFile = false;
 	bool bReadOnlyChannelsFrozen = false;
+	bool bHasPanicked = false;
 	FString PausedPreset;
 	FDelegateHandle EndFrameUpdateHandle;
 	FDelegateHandle OnGetOnScreenMessageHandle;
@@ -563,6 +566,11 @@ uint32 FTraceAuxiliaryImpl::HashChannelName(const TCHAR* Name)
 ////////////////////////////////////////////////////////////////////////////////
 bool FTraceAuxiliaryImpl::Connect(FTraceAuxiliary::EConnectionType Type, const TCHAR* Parameter, const FTraceAuxiliary::FLogCategoryAlias& LogCategory, uint16 SendFlags)
 {
+	if (bHasPanicked)
+	{
+		UE_LOG_REF(LogCategory, Warning, TEXT("Trace has panicked, cannot start new trace."));
+		return false;
+	}
 	// Connect to trace server or write to a file, but only if we're not already tracing.
 	if (UE::Trace::IsTracing())
 	{
@@ -628,6 +636,11 @@ bool FTraceAuxiliaryImpl::Connect(FTraceAuxiliary::EConnectionType Type, const T
 bool FTraceAuxiliaryImpl::ConnectRelay(UPTRINT Handle, UE::Trace::IoWriteFunc WriteFunc, UE::Trace::IoCloseFunc CloseFunc, uint16 SendFlags)
 {
 	check(WriteFunc && CloseFunc);
+	if (bHasPanicked)
+	{
+		UE_LOG(LogTrace, Warning, TEXT("Trace has panicked, cannot start new trace."));
+		return false;
+	}
 
 	if (UE::Trace::IsTracing())
 	{
@@ -737,6 +750,11 @@ bool FTraceAuxiliaryImpl::InitSecure()
 ////////////////////////////////////////////////////////////////////////////////
 bool FTraceAuxiliaryImpl::ConnectSecure(FTraceAuxiliary::FSecureOptions& InOptions, uint16 SendFlags)
 {
+	if (bHasPanicked)
+	{
+		UE_LOG(LogTrace, Warning, TEXT("Trace has panicked, cannot start new trace."));
+		return false;
+	}
 	UE::Trace::FSecureTraceOptions SecureOptions;
 	if (!InOptions.AuthorizationToken.IsEmpty())
 	{
@@ -790,8 +808,13 @@ uint16 FTraceAuxiliaryImpl::GetSecurePort() const
 #endif // UE_TRACE_ALLOW_SECURE_TRACING
 
 ////////////////////////////////////////////////////////////////////////////////
-bool FTraceAuxiliaryImpl::Stop()
+bool FTraceAuxiliaryImpl::Stop(const FTraceAuxiliary::FLogCategoryAlias& LogCategory)
 {
+	if (bHasPanicked)
+	{
+		UE_LOG_REF(LogCategory, Warning, TEXT("Trace has panicked, cannot stop trace."));
+		return false;
+	}
 	if (IsParentProcessAndPreFork())
 	{
 		return false;
@@ -1401,6 +1424,7 @@ void TraceAuxiliaryOnMessageCallback(const UE::Trace::FMessageEvent& Message)
 		break;
 	case UE::Trace::EMessageType::PanicError:
 		LOG_WITH_VERBOSITY(Warning, Msg);
+		GTraceAuxiliary.SetHasPanicked();
 		break;
 	default:
 		{
@@ -1542,7 +1566,7 @@ static void TraceAuxiliaryStart(const TArray<FString>& Args)
 static void TraceAuxiliaryStop()
 {
 	UE_LOG(LogConsoleResponse, Display, TEXT("Tracing stopped."));
-	GTraceAuxiliary.Stop();
+	GTraceAuxiliary.Stop(LogConsoleResponse);
 }
 
 ////////////////////////////////////////////////////////////////////////////////
@@ -1563,30 +1587,36 @@ static void TraceAuxiliaryResume()
 static void TraceAuxiliaryStatus()
 {
 	UE_LOG(LogConsoleResponse, Display, TEXT("Trace status ----------------------------------------------------------"));
-
-	// Status of data connection
-	TStringBuilder<256> ConnectionStr;
-	FGuid SessionGuid, TraceGuid;
-	if (GTraceAuxiliary.IsConnected(SessionGuid, TraceGuid))
+	if (GTraceAuxiliary.GetHasPanicked())
+	{
+		UE_LOG(LogConsoleResponse, Error, TEXT("Trace is in panic mode and no longer useable."));
+	}
+	else
 	{
-		const FString Dest = GTraceAuxiliary.GetDest();
-		if (!Dest.IsEmpty())
+		// Status of data connection
+		TStringBuilder<256> ConnectionStr;
+		FGuid SessionGuid, TraceGuid;
+		if (GTraceAuxiliary.IsConnected(SessionGuid, TraceGuid))
 		{
-			ConnectionStr.Appendf(TEXT("Tracing to '%s', "), *Dest);
-			ConnectionStr.Appendf(TEXT("session %s trace %s"), *SessionGuid.ToString(), *TraceGuid.ToString());
+			const FString Dest = GTraceAuxiliary.GetDest();
+			if (!Dest.IsEmpty())
+			{
+				ConnectionStr.Appendf(TEXT("Tracing to '%s', "), *Dest);
+				ConnectionStr.Appendf(TEXT("session %s trace %s"), *SessionGuid.ToString(), *TraceGuid.ToString());
+			}
+			else
+			{
+				// If GTraceAux doesn't know about the target but we are still tracing this is an externally initiated connection
+				// (e.g. connection command from Insights).
+				ConnectionStr = TEXT("Tracing to unknown target (externally set)");
+			}
 		}
 		else
 		{
-			// If GTraceAux doesn't know about the target but we are still tracing this is an externally initiated connection
-			// (e.g. connection command from Insights).
-			ConnectionStr = TEXT("Tracing to unknown target (externally set)");
+			ConnectionStr = TEXT("Not tracing");
 		}
+		UE_LOG(LogConsoleResponse, Display, TEXT("Connection: %s"), ConnectionStr.ToString());
 	}
-	else
-	{
-		ConnectionStr = TEXT("Not tracing");
-	}
-	UE_LOG(LogConsoleResponse, Display, TEXT("Connection: %s"), ConnectionStr.ToString());
 
 	// Stats
 	UE::Trace::FStatistics Stats;
@@ -2103,7 +2133,7 @@ uint16 FTraceAuxiliary::GetSecurePort()
 bool FTraceAuxiliary::Stop()
 {
 #if UE_TRACE_MINIMAL_ENABLED
-	return GTraceAuxiliary.Stop();
+	return GTraceAuxiliary.Stop(LogTrace);
 #else
 	return false;
 #endif
