Added hack to force non-scoped enums in dynamic delegate signatures to be a byte.  This should be caught in UHT, and the hack can be reversed once it is rejected fully.

#Synced-CL 51823540

diff --git a/Engine/Source/Runtime/Core/Public/Concepts/UnscopedEnum.h b/Engine/Source/Runtime/Core/Public/Concepts/UnscopedEnum.h
new file mode 100644
--- /dev/null
+++ b/Engine/Source/Runtime/Core/Public/Concepts/UnscopedEnum.h
@@ -0,0 +1,15 @@
+// Copyright Epic Games, Inc. All Rights Reserved.
+
+#pragma once
+
+#include "Concepts/ConvertibleTo.h"
+#include "Concepts/Enum.h"
+
+namespace UE
+{
+	/**
+	 * Concept which describes an unscoped enum type (aka. non-class enum).
+	 */
+	template <typename T>
+	concept CUnscopedEnum = UE::CEnum<T> && CConvertibleTo<T, int>;
+}
diff --git a/Engine/Source/Runtime/Core/Public/UObject/ScriptDelegates.h b/Engine/Source/Runtime/Core/Public/UObject/ScriptDelegates.h
--- a/Engine/Source/Runtime/Core/Public/UObject/ScriptDelegates.h
+++ b/Engine/Source/Runtime/Core/Public/UObject/ScriptDelegates.h
@@ -6,9 +6,10 @@
 #include "Concepts/Const.h"
 #include "Concepts/Reference.h"
 #include "Concepts/Scalar.h"
+#include "Concepts/UnscopedEnum.h" // this can be removed when UE_DYNAMIC_DELEGATE_UNSCOPED_ENUM_HACK is removed
 #include "Containers/Array.h"
 #include "Containers/ContainerAllocationPolicies.h"
-#include "Containers/EnumAsByte.h"
+#include "Containers/EnumAsByte.h" // this can be removed when UE_DYNAMIC_DELEGATE_UNSCOPED_ENUM_HACK is removed
 #include "Containers/UnrealString.h"
 #include "PropertyPortFlags.h"
 #include "Delegates/DelegateAccessHandler.h"
@@ -29,6 +30,16 @@
 	#endif
 #endif
 
+/* A hack to force unscoped enums to be a TEnumAsByte.  This should be enforced by UHT but it wasn't being caught
+ * for function parameters - including dynamic delegate declarations - so we force it here, otherwise the ProcessEvent
+ * struct can end up misaligned.
+ *
+ * This can be removed when UHT rejects unscoped enum parameters in dynamic delegate declarations.
+ */
+#ifndef UE_DYNAMIC_DELEGATE_UNSCOPED_ENUM_HACK
+	#define UE_DYNAMIC_DELEGATE_UNSCOPED_ENUM_HACK 1
+#endif
+
 #if UE_USE_DYNAMIC_DELEGATE_PAYLOADS
 	#include "Concepts/Arithmetic.h"
 	#include "Concepts/UObject.h"
@@ -230,12 +241,26 @@ namespace UE::Core::Private
 	template <typename T>
 	struct TDelegateWrapperReturnBase
 	{
-		std::remove_cv_t<std::remove_reference_t<T>> ReturnValue = {};
+		using MemberType = std::remove_cv_t<std::remove_reference_t<T>>;
+
+#if UE_DYNAMIC_DELEGATE_UNSCOPED_ENUM_HACK
+		using RemappedMemberType = std::conditional_t<UE::CUnscopedEnum<MemberType>, TEnumAsByte<MemberType>, MemberType>;
+		RemappedMemberType ReturnValue = {};
+#else
+		MemberType ReturnValue = {};
+#endif
 	};
 	template <CForceInitable T>
 	struct TDelegateWrapperReturnBase<T>
 	{
-		std::remove_cv_t<std::remove_reference_t<T>> ReturnValue = T(ForceInit);
+		using MemberType = std::remove_cv_t<std::remove_reference_t<T>>;
+
+#if UE_DYNAMIC_DELEGATE_UNSCOPED_ENUM_HACK
+		using RemappedMemberType = std::conditional_t<UE::CUnscopedEnum<MemberType>, TEnumAsByte<MemberType>, MemberType>;
+		RemappedMemberType ReturnValue = T(ForceInit);
+#else
+		MemberType ReturnValue = T(ForceInit);
+#endif
 	};
 
 	template <typename T, uint32 Index>
@@ -243,11 +268,12 @@ namespace UE::Core::Private
 	{
 		using MemberType = std::remove_cv_t<std::remove_reference_t<T>>;
 
-		// HACK - this forces non-scoped enums to be bytes.  This is expected in dynamic delegates and should be rejected by UHT,
-		// but hasn't been, though it did force them to bytes in its generated Params struct passed to ProcessEvent.
-		using RemappedMemberType = std::conditional_t<std::is_enum_v<MemberType> && std::is_convertible_v<MemberType, int>, TEnumAsByte<MemberType>, MemberType>;
-
+#if UE_DYNAMIC_DELEGATE_UNSCOPED_ENUM_HACK
+		using RemappedMemberType = std::conditional_t<UE::CUnscopedEnum<MemberType>, TEnumAsByte<MemberType>, MemberType>;
 		RemappedMemberType Member;
+#else
+		MemberType Member;
+#endif
 
 		// This odd constraint is to allow selection of the correct MoveOut overload by specifying the
 		// index when they are merged into the derived class's overload set.
diff --git a/Engine/Source/Runtime/CoreUObject/Public/UObject/DynamicDelegateTestObject.cpp b/Engine/Source/Runtime/CoreUObject/Public/UObject/DynamicDelegateTestObject.cpp
--- a/Engine/Source/Runtime/CoreUObject/Public/UObject/DynamicDelegateTestObject.cpp
+++ b/Engine/Source/Runtime/CoreUObject/Public/UObject/DynamicDelegateTestObject.cpp
@@ -451,6 +451,13 @@ void UDynamicDelegateTestObject::ConstFuncOut5(int& Int, FString& Str, FName& Na
 	WeakObjPtr = GWeakObjPtrOut;
 }
 
+void UDynamicDelegateTestObject::TestUnscopedEnum(int16 Int1, ETestUnscopedEnum::Type E, int16 Int2)
+{
+	CHECK(Int1 == 5);
+	CHECK(E == ETestUnscopedEnum::E2);
+	CHECK(Int2 == 10);
+}
+
 #define DEFINE_ARG_VARIABLES_1(Arg1Val)                                              auto Arg1 = Arg1Val;
 #define DEFINE_ARG_VARIABLES_2(Arg1Val, Arg2Val)                                     auto Arg1 = Arg1Val; auto Arg2 = Arg2Val;
 #define DEFINE_ARG_VARIABLES_3(Arg1Val, Arg2Val, Arg3Val)                            auto Arg1 = Arg1Val; auto Arg2 = Arg2Val; auto Arg3 = Arg3Val;
@@ -979,6 +986,13 @@ void UDynamicDelegateTestObject::RunTests()
 		BIND_CALL_CHECK_DELEGATE_OUT(FTestDynamicDelegateOut5, 0, &UDynamicDelegateTestObject::ConstFuncOut5, (GIntArg, GStrArg, GNameArg, bGBoolArg, GWeakObjPtrArg), (GIntOut, GStrOut, GNameOut, bGBoolOut, GWeakObjPtrOut));
 		BIND_CALL_CHECK_DELEGATE_OUT(FTestDynamicMulticastDelegateOut5, 1, &UDynamicDelegateTestObject::ConstFuncOut5, (GIntArg, GStrArg, GNameArg, bGBoolArg, GWeakObjPtrArg), (GIntOut, GStrOut, GNameOut, bGBoolOut, GWeakObjPtrOut));
 	}
+
+	// Unscoped enum
+	{
+		FTestUnscopedEnumDelegate Test;
+		Test.BindDynamic(this, &UDynamicDelegateTestObject::TestUnscopedEnum);
+		Test.Execute(5, ETestUnscopedEnum::E2, 10);
+	}
 }
 
 TEST_CASE_NAMED(DynamicDelegateTest, "UE::CoreUObject::DynamicDelegateTest", "[ApplicationContextMask][SmokeFilter]")
diff --git a/Engine/Source/Runtime/CoreUObject/Public/UObject/DynamicDelegateTestObject.h b/Engine/Source/Runtime/CoreUObject/Public/UObject/DynamicDelegateTestObject.h
--- a/Engine/Source/Runtime/CoreUObject/Public/UObject/DynamicDelegateTestObject.h
+++ b/Engine/Source/Runtime/CoreUObject/Public/UObject/DynamicDelegateTestObject.h
@@ -54,6 +54,18 @@ DECLARE_DYNAMIC_MULTICAST_DELEGATE_ThreeParams(FTestDynamicMulticastDelegateOut3
 DECLARE_DYNAMIC_MULTICAST_DELEGATE_FourParams (FTestDynamicMulticastDelegateOut4, int&, Int, FString&, Str, FName&, Name, bool&, bBool);
 DECLARE_DYNAMIC_MULTICAST_DELEGATE_FiveParams (FTestDynamicMulticastDelegateOut5, int&, Int, FString&, Str, FName&, Name, bool&, bBool, TWeakObjectPtr<UDynamicDelegateTestObject>&, WeakPtr);
 
+UENUM()
+namespace ETestUnscopedEnum
+{
+	enum Type : int
+	{
+		E1,
+		E2
+	};
+}
+
+DECLARE_DYNAMIC_DELEGATE_ThreeParams(FTestUnscopedEnumDelegate, int16, Int1, ETestUnscopedEnum::Type, E, int16, Int2);
+
 UCLASS()
 class UDynamicDelegateTestObject : public UObject
 {
@@ -124,6 +136,8 @@ public:
 	UFUNCTION() void ConstFuncOut4(int& Int, FString& Str, FName& Name, bool& bBool) const;
 	UFUNCTION() void ConstFuncOut5(int& Int, FString& Str, FName& Name, bool& bBool, TWeakObjectPtr<UDynamicDelegateTestObject>& WeakPtr) const;
 
+	UFUNCTION() void TestUnscopedEnum(int16 Int1, ETestUnscopedEnum::Type E, int16 Int2);
+
 	void RunTests();
 };
 
