Logging: Ignore SetColor verbosity in FOutputDeviceDebug

#Synced-CL 50967057

diff --git a/Engine/Source/Runtime/Core/Private/Misc/OutputDeviceDebug.cpp b/Engine/Source/Runtime/Core/Private/Misc/OutputDeviceDebug.cpp
--- a/Engine/Source/Runtime/Core/Private/Misc/OutputDeviceDebug.cpp
+++ b/Engine/Source/Runtime/Core/Private/Misc/OutputDeviceDebug.cpp
@@ -66,6 +66,11 @@ void FOutputDeviceDebug::Serialize(const TCHAR* Data, ELogVerbosity::Type Verbos
 
 void FOutputDeviceDebug::SerializeRecord(const UE::FLogRecord& Record)
 {
+	if (Record.GetVerbosity() == ELogVerbosity::SetColor)
+	{
+		return;
+	}
+
 	ConditionalTickAsync(FPlatformTime::Seconds() - GStartTime);
 
 	if (bSerializeAsJson.load(std::memory_order_relaxed))
