[Android] Update FdSan configrule check to new early API

#rnx
#Synced-CL 51996919

#Synced-CL 51996919

diff --git a/Engine/Source/Runtime/Core/Private/Android/AndroidPlatformFile.cpp b/Engine/Source/Runtime/Core/Private/Android/AndroidPlatformFile.cpp
--- a/Engine/Source/Runtime/Core/Private/Android/AndroidPlatformFile.cpp
+++ b/Engine/Source/Runtime/Core/Private/Android/AndroidPlatformFile.cpp
@@ -83,14 +83,7 @@ namespace UE::Android::Private
 	{
 		FdSanConfig()
 		{
-			if (const FString* ConfigRulesOverride = FPlatformMisc::GetConfigRulesVariable(TEXT("UseFdSan")))
-			{
-				bIsEnabled = ConfigRulesOverride->Equals("true", ESearchCase::IgnoreCase);
-			}
-			else
-			{
-				bIsEnabled = true;
-			}
+			bIsEnabled = !FPlatformMisc::GetConfigRulesVariableAsUtf8("UseFdSan").Equals("false", ESearchCase::IgnoreCase);
 		}
 
 		bool bIsEnabled;
