[TaskGraph]
- Fix potential stack overflow for very long chain of events

#Synced-CL 51934851

diff --git a/Engine/Source/Runtime/Core/Private/Tasks/TaskPrivate.cpp b/Engine/Source/Runtime/Core/Private/Tasks/TaskPrivate.cpp
--- a/Engine/Source/Runtime/Core/Private/Tasks/TaskPrivate.cpp
+++ b/Engine/Source/Runtime/Core/Private/Tasks/TaskPrivate.cpp
@@ -302,6 +302,93 @@ namespace UE::Tasks
 		{
 			GetPipe()->ClearTask(*this);
 		}
+	
+		// Stack overflow prevention for deep Close -> TryUnlock -> Close chains.
+		// Track recursion depth per-thread; beyond MaxCloseRecursionDepth, defer
+		// subsequent unlocks into a worklist drained iteratively at the outermost level.
+		static constexpr int32 MaxCloseRecursionDepth = 128;
+
+		struct FCloseOverflowState
+		{
+			int32 Depth = 0;
+			TArray<FTaskBase*> Pending;
+		};
+		static thread_local FCloseOverflowState CloseOverflow;
+
+		void FTaskBase::Close()
+		{
+			TASKGRAPH_VERBOSE_EVENT_SCOPE(FTaskBase::Close);
+			checkSlow(!IsCompleted());
+
+			++CloseOverflow.Depth;
+
+			// Push the first subsequent to the local queue so we pick it up directly as our next task.
+			// This saves us the cost of going to the global queue and performing a wake-up.
+			// But if we're a task event, always wake up new workers because the current task could continue executing for a long time after the trigger.
+			bool bWakeUpWorker = ExtendedPriority == EExtendedTaskPriority::TaskEvent;
+
+			auto LocalSubsequents = Subsequents.Close();
+
+			if (LIKELY(CloseOverflow.Depth <= MaxCloseRecursionDepth))
+			{
+				for (FTaskBase* Subsequent : LocalSubsequents)
+				{
+					// bWakeUpWorker is passed by reference and is automatically set to true if we successfully schedule a task on the local queue.
+					// so all the remaining ones are sent to the global queue.
+					Subsequent->TryUnlock(bWakeUpWorker);
+				}
+			}
+			else
+			{
+				for (FTaskBase* Subsequent : LocalSubsequents)
+				{
+					// Defer unlock to avoid stack overflow from deep recursion
+					CloseOverflow.Pending.Add(Subsequent);
+				}
+			}
+
+			// Clear the pipe after the task is completed (subsequents closed) so that any tasks part of the
+			// pipe are not seen still being executed after FPipe::WaitUntilEmpty has returned.
+			if (GetPipe() != nullptr)
+			{
+				ClearPipe();
+			}
+
+			// release nested tasks
+			ReleasePrerequisites();
+
+			TaskTrace::Completed(GetTraceId());
+
+			// In case a thread is waiting on us to perform retraction, now is the time to try retraction again.
+			StateChangeEvent.NotifyWeak();
+
+			// The outermost Close on this thread drains any deferred unlocks iteratively.
+			if (--CloseOverflow.Depth == 0 && CloseOverflow.Pending.Num() > 0)
+			{
+				// Always wake up workers, at this point we're long past the first element anyway
+				bWakeUpWorker = true;
+
+				TArray<FTaskBase*> Pending;
+
+				// Bump Depth back to 1 during drain so recursive Close calls don't also try to drain.
+				++CloseOverflow.Depth;
+				while (CloseOverflow.Pending.Num())
+				{
+					// Backbuffer to maintain ordering without having to allocate for each iteration
+					Swap(Pending, CloseOverflow.Pending);
+
+					for (FTaskBase* Task : Pending)
+					{
+						// TryUnlock may recurse into Close, appending more items
+						Task->TryUnlock(bWakeUpWorker);
+					}
+
+					Pending.Reset(); 
+				}
+				CloseOverflow.Pending.Empty(); // Don't leave an allocation when done
+				--CloseOverflow.Depth;
+			}
+		}
 
 		static thread_local FTaskBase* CurrentTask = nullptr;
 
diff --git a/Engine/Source/Runtime/Core/Public/Tasks/TaskPrivate.h b/Engine/Source/Runtime/Core/Public/Tasks/TaskPrivate.h
--- a/Engine/Source/Runtime/Core/Public/Tasks/TaskPrivate.h
+++ b/Engine/Source/Runtime/Core/Public/Tasks/TaskPrivate.h
@@ -546,38 +546,7 @@ public:
 			}
 
 			// closes task by unlocking its subsequents and flagging it as completed
-			void Close()
-			{
-				TASKGRAPH_VERBOSE_EVENT_SCOPE(FTaskBase::Close);
-				checkSlow(!IsCompleted());
-				
-				// Push the first subsequent to the local queue so we pick it up directly as our next task.
-				// This saves us the cost of going to the global queue and performing a wake-up.
-				// But if we're a task event, always wake up new workers because the current task could continue executing for a long time after the trigger.
-				bool bWakeUpWorker = ExtendedPriority == EExtendedTaskPriority::TaskEvent;
-
-				for (FTaskBase* Subsequent : Subsequents.Close())
-				{
-					// bWakeUpWorker is passed by reference and is automatically set to true if we successfully schedule a task on the local queue.
-					// so all the remaining ones are sent to the global queue.
-					Subsequent->TryUnlock(bWakeUpWorker);
-				}
-
-				// Clear the pipe after the task is completed (subsequents closed) so that any tasks part of the
-				// pipe are not seen still being executed after FPipe::WaitUntilEmpty has returned.
-				if (GetPipe() != nullptr)
-				{
-					ClearPipe();
-				}
-
-				// release nested tasks
-				ReleasePrerequisites();
-
-				TaskTrace::Completed(GetTraceId());
-
-				// In case a thread is waiting on us to perform retraction, now is the time to try retraction again.
-				StateChangeEvent.NotifyWeak();
-			}
+			CORE_API void Close();
 
 			CORE_API void ClearPipe();
 
diff --git a/Engine/Source/Runtime/Core/Tests/Tasks/TasksTest.cpp b/Engine/Source/Runtime/Core/Tests/Tasks/TasksTest.cpp
--- a/Engine/Source/Runtime/Core/Tests/Tasks/TasksTest.cpp
+++ b/Engine/Source/Runtime/Core/Tests/Tasks/TasksTest.cpp
@@ -1920,6 +1920,42 @@ namespace UE { namespace TasksTests
 		}
 	}
 
+	TEST_CASE_NAMED(FTaskGraphLongEventChain, "System::Core::Async::Tasks::LongEventChain", "[ApplicationContextMask][EngineFilter]")
+	{
+		UE::Tasks::FTaskEvent RootEvent(TEXT("RootEvent"));
+
+		UE::Tasks::FPipe Pipe(TEXT("Pipe"));
+
+		std::atomic<int32> CurrentIndex = 0;
+
+		// Launch a long event chain with RootEvent not triggered yet.
+		// Each event will unlock a task to put in the pipe, and then we can observe the unlock order.
+		UE::Tasks::FTaskEvent LastEvent = RootEvent;
+		for (int32 Index = 0; Index < 100000; ++Index)
+		{
+			UE::Tasks::FTaskEvent Event(TEXT("Event"));
+			Event.AddPrerequisites(LastEvent);
+
+			Pipe.Launch(
+				TEXT("Task"),
+				[Index, &CurrentIndex]
+				{
+					CHECK(Index == CurrentIndex++);
+				},
+				Event
+			);
+
+			// RootEvent is not triggered yet, so this does launch the task yet
+			Event.Trigger();
+			LastEvent = Event;
+		}
+
+		// Without the overflow handling in Close(), this would probably cause a stack overflow
+		RootEvent.Trigger();
+
+		Pipe.WaitUntilEmpty();
+	}
+
 #if PLATFORM_SUPPORTS_ASYMMETRIC_FENCES
 	TEST_CASE_NAMED(FAsymmetricThreadFence, "System::Core::Async::AsymmetricThreadFence", "[.][ApplicationContextMask][EngineFilter]")
 	{
