[RaceDetector]
- Add SuspendThread barrier to prevent false positives during callstack captures
- Fix reset shadow that could be starved by concurrent threads doing heavy access to virtualalloc
- Fix false positive on TLS memory reused between threads

https://horde.devtools.epicgames.com/job/69c038ad522e8c7766bc55ca
#Synced-CL 51934851

diff --git a/Engine/Source/Runtime/Core/Private/Sanitizer/RaceDetector.cpp b/Engine/Source/Runtime/Core/Private/Sanitizer/RaceDetector.cpp
--- a/Engine/Source/Runtime/Core/Private/Sanitizer/RaceDetector.cpp
+++ b/Engine/Source/Runtime/Core/Private/Sanitizer/RaceDetector.cpp
@@ -160,6 +160,7 @@ namespace UE::Sanitizer::RaceDetector {
 	UE::FTimeout GracePeriodTimeout = UE::FTimeout::Never();
 	volatile uint32 GracePeriodEpoch = 0;
 
+	thread_local bool bIsResettingShadowThread = false;
 	volatile bool bIsResettingShadow = false;
 	volatile bool bIsDebuggerPresent = false;
 	volatile bool bDetailedLogGlobal = false;
@@ -314,6 +315,19 @@ namespace UE::Sanitizer::RaceDetector {
 		return ContextMapping.FindRef(ContextId);
 	}
 
+	INSTRUMENTATION_FUNCTION_ATTRIBUTES TRefCountPtr<FContext> GetContextByThreadId(uint32 ThreadId)
+	{
+		TWriteScopeLock Scope(ContextMappingLock);
+		for (auto& [Key, Context] : ContextMapping)
+		{
+			if (Context->ThreadId == ThreadId)
+			{
+				return Context;
+			}
+		}
+		return nullptr;
+	}
+
 	FORCEINLINE INSTRUMENTATION_FUNCTION_ATTRIBUTES UPTRINT GetAlignedAddress(UPTRINT Ptr)
 	{
 		return UPTRINT(Ptr) & ~UPTRINT(0b111);
@@ -1093,6 +1107,7 @@ namespace UE::Sanitizer::RaceDetector {
 		FPlatformMisc::LowLevelOutputDebugStringf(TEXT("[%d] Reset Shadow Memory Started\n"), Platform::GetCurrentThreadId());
 
 		FPlatformAtomics::InterlockedExchange((volatile int8*)&bIsResettingShadow, (int8)1);
+		bIsResettingShadowThread = true;
 
 		if (RaceDetectorActive)
 		{
@@ -1139,6 +1154,7 @@ namespace UE::Sanitizer::RaceDetector {
 			HookInstrumentation();
 		}
 		bIsResettingShadow = false;
+		bIsResettingShadowThread = false;
 	}
 
 	static TUniquePtr<FArchive> RaceDetectorLog;
@@ -1231,6 +1247,8 @@ namespace UE::Sanitizer::RaceDetector {
 
 		HandleReports();
 	}
+	
+	namespace Platform { extern void ThrottleMemoryAccessesDuringShadowReset(); }
 
 	/**
 	 * FMalloc proxy removes instrumentation for allocators 
@@ -1261,6 +1279,11 @@ namespace UE::Sanitizer::RaceDetector {
 		INSTRUMENTATION_FUNCTION_ATTRIBUTES virtual void* Malloc(SIZE_T Size, uint32 Alignment) override
 		{
 			FInstrumentationScope Scope;
+
+			// Reduce the pressure on the memory allocator while we're resetting shadow as it can starve
+			// the reset thread completely
+			Platform::ThrottleMemoryAccessesDuringShadowReset();
+
 			return InnerMalloc->Malloc(Size, Alignment);
 		}
 
@@ -1272,6 +1295,10 @@ namespace UE::Sanitizer::RaceDetector {
 				return InnerMalloc->Realloc(InPtr, NewSize, Alignment);
 			}
 
+			// Reduce the pressure on the memory allocator while we're resetting shadow as it can starve
+			// the reset thread completely
+			Platform::ThrottleMemoryAccessesDuringShadowReset();
+
 			// We have to always allocate new blocs in order to invalidate the old memory range before it can be reused.
 			void* NewPtr = nullptr;
 			if (NewSize)
@@ -1301,6 +1328,11 @@ namespace UE::Sanitizer::RaceDetector {
 		INSTRUMENTATION_FUNCTION_ATTRIBUTES virtual void Free(void* Ptr) override
 		{
 			FInstrumentationScope Scope;
+
+			// Reduce the pressure on the memory allocator while we're resetting shadow as it can starve
+			// the reset thread completely
+			Platform::ThrottleMemoryAccessesDuringShadowReset();
+
 			SIZE_T Size = 0;
 			InnerMalloc->GetAllocationSize(Ptr, Size);
 			if (Ptr && Size)
diff --git a/Engine/Source/Runtime/Core/Private/Sanitizer/RaceDetectorWindows.cpp b/Engine/Source/Runtime/Core/Private/Sanitizer/RaceDetectorWindows.cpp
--- a/Engine/Source/Runtime/Core/Private/Sanitizer/RaceDetectorWindows.cpp
+++ b/Engine/Source/Runtime/Core/Private/Sanitizer/RaceDetectorWindows.cpp
@@ -26,12 +26,122 @@ using namespace UE::Sanitizer::RaceDetector;
 namespace UE::Sanitizer::RaceDetector {
 	extern void InitializeBoot();
 	extern void SanitizerThreadRun(volatile bool& bContinue);
+	extern TRefCountPtr<FContext> GetContextByThreadId(uint32 InThreadId);
 	extern volatile bool bIsResettingShadow;
+	extern thread_local bool bIsResettingShadowThread;
 	extern void WriteToLog(FString&& Message);
 	extern bool bRuntimeInitialized;
 	constexpr int32 MinStackSize = 256 * 1024;
 }
 
+// Sentinel used to validate our TLS walking assumptions at runtime.
+static __declspec(thread) volatile uint8 TlsValidationSentinel = 0;
+
+// Tri-state: 0 = not yet checked, 1 = validated OK, -1 = validation failed.
+static volatile int8 TlsCleanupState = 0;
+
+// Walks every loaded module's PE TLS directory and invokes the callback with the
+// base address and size of each per-thread TLS block. No API calls are made —
+// everything is read directly from TEB/PEB/PE structures to avoid reentrancy
+// through detoured functions.
+// Returns false if the TEB or PEB structures couldn't be read.
+template<typename CallbackType>
+INSTRUMENTATION_FUNCTION_ATTRIBUTES static bool ForEachTlsRegion(CallbackType&& Callback)
+{
+	// TEB offset 0x58 (x64) is ThreadLocalStoragePointer
+	void** TlsArray = *(void***)((BYTE*)NtCurrentTeb() + 0x58);
+	if (!TlsArray)
+	{
+		return false;
+	}
+
+	PEB* Peb = NtCurrentTeb()->ProcessEnvironmentBlock;
+	if (!Peb || !Peb->Ldr)
+	{
+		return false;
+	}
+
+	LIST_ENTRY* Head = &Peb->Ldr->InMemoryOrderModuleList;
+	for (LIST_ENTRY* Entry = Head->Flink; Entry != Head; Entry = Entry->Flink)
+	{
+		LDR_DATA_TABLE_ENTRY* LdrEntry = CONTAINING_RECORD(Entry, LDR_DATA_TABLE_ENTRY, InMemoryOrderLinks);
+		BYTE* DllBase = (BYTE*)LdrEntry->DllBase;
+		if (!DllBase) 
+		{
+			continue;
+		}
+
+		IMAGE_DOS_HEADER* DosHeader = (IMAGE_DOS_HEADER*)DllBase;
+		if (DosHeader->e_magic != IMAGE_DOS_SIGNATURE) 
+		{
+			continue;
+		}
+		IMAGE_NT_HEADERS64* NtHeaders = (IMAGE_NT_HEADERS64*)(DllBase + DosHeader->e_lfanew);
+		if (NtHeaders->Signature != IMAGE_NT_SIGNATURE) 
+		{
+			continue;
+		}
+
+		IMAGE_DATA_DIRECTORY& TlsDataDir = NtHeaders->OptionalHeader.DataDirectory[IMAGE_DIRECTORY_ENTRY_TLS];
+		if (TlsDataDir.VirtualAddress == 0) 
+		{
+			continue;
+		}
+
+		IMAGE_TLS_DIRECTORY* TlsDir = (IMAGE_TLS_DIRECTORY*)(DllBase + TlsDataDir.VirtualAddress);
+		DWORD TlsIndex = *(DWORD*)TlsDir->AddressOfIndex;
+		SIZE_T TlsSize = TlsDir->EndAddressOfRawData - TlsDir->StartAddressOfRawData;
+		void* TlsBase = TlsArray[TlsIndex];
+
+		if (TlsBase && TlsSize > 0)
+		{
+			Callback(TlsBase, TlsSize);
+		}
+	}
+
+	return true;
+}
+
+// Validates that ForEachTlsRegion produces correct results by checking that
+// a known __declspec(thread) variable falls within one of the discovered regions.
+INSTRUMENTATION_FUNCTION_ATTRIBUTES static bool ValidateTlsCleanupAssumptions()
+{
+	const BYTE* SentinelAddr = (const BYTE*)&TlsValidationSentinel;
+	bool bFound = false;
+
+	bool bSuccess = ForEachTlsRegion([&](void* TlsBase, SIZE_T TlsSize)
+	{
+		const BYTE* Base = (const BYTE*)TlsBase;
+		if (SentinelAddr >= Base && SentinelAddr < Base + TlsSize)
+		{
+			bFound = true;
+		}
+	});
+
+	return bSuccess && bFound;
+}
+
+// Ensures TLS cleanup has been validated. Called lazily on the first thread
+// shutdown so that all modules are loaded and TLS is fully initialized.
+INSTRUMENTATION_FUNCTION_ATTRIBUTES static bool EnsureTlsCleanupValidated()
+{
+	if (TlsCleanupState == 0)
+	{
+		if (ValidateTlsCleanupAssumptions())
+		{
+			TlsCleanupState = 1;
+		}
+		else
+		{
+			TlsCleanupState = -1;
+			UE_LOG(LogRaceDetector, Warning,
+				TEXT("TLS shadow cleanup disabled: could not validate TEB/PE TLS assumptions on this OS version. ")
+				TEXT("TLS-related false positives may occur."));
+		}
+	}
+	return TlsCleanupState > 0;
+}
+
 #define START_WINAPI_INSTRUMENTATION(DoIfNotInstrumenting) \
 	DWORD OriginalError = GetLastError(); \
 	FContext* ContextPtr = GetThreadContext(); \
@@ -125,6 +235,16 @@ INSTRUMENTATION_FUNCTION_ATTRIBUTES void InstrumentEndThread()
 	GetCurrentThreadStackLimits(&LowLimit, &HighLimit);
 	UE::Sanitizer::RaceDetector::FreeMemoryRange(reinterpret_cast<void*>(LowLimit), HighLimit - LowLimit);
 
+	// Clean TLS to avoid false positives when memory gets reused for another thread.
+	// Validated lazily on the first thread shutdown to ensure all modules are loaded.
+	if (EnsureTlsCleanupValidated())
+	{
+		ForEachTlsRegion([](void* TlsBase, SIZE_T TlsSize)
+		{
+			UE::Sanitizer::RaceDetector::FreeMemoryRange(TlsBase, TlsSize);
+		});
+	}
+
 	ReleaseCurrentContext();
 
 	// We released our context so we can't decrement the count now but we can't decrement it
@@ -654,6 +774,37 @@ INSTRUMENTATION_FUNCTION_ATTRIBUTES PVOID NTAPI DetouredRtlReallocateHeap(PVOID
 	return NewPtr;
 }
 
+DWORD SuspendThread(
+  [in] HANDLE hThread
+);
+
+typedef NTSYSAPI DWORD(NTAPI SUSPENDTHREAD)(IN HANDLE ThreadHandle);
+
+SUSPENDTHREAD* TrueSuspendThread = nullptr;
+INSTRUMENTATION_FUNCTION_ATTRIBUTES DWORD NTAPI DetouredSuspendThread(HANDLE ThreadHandle)
+{
+	START_WINAPI_INSTRUMENTATION(return TrueSuspendThread(ThreadHandle));
+
+	// Suspend the thread first, then acquire its clock state. Once SuspendThread returns
+	// successfully, the target thread is stopped and all its prior memory accesses are
+	// visible. This is needed for callstack capture patterns where thread A suspends
+	// thread B, reads B's context/stack, then resumes B.
+	DWORD result = TrueSuspendThread(ThreadHandle);
+	if (result != (DWORD)-1)
+	{
+		DWORD ThreadId = ::GetThreadId(ThreadHandle);
+		TRefCountPtr<FContext> SuspendedThreadContext = GetContextByThreadId(ThreadId);
+		if (SuspendedThreadContext)
+		{
+			Context.ClockBank.Acquire(SuspendedThreadContext->ClockBank, _ReturnAddress());
+		}
+	}
+
+	FINISH_WINAPI_INSTRUMENTATION();
+
+	return result;
+}
+
 DWORD ResumeThread(
   [in] HANDLE hThread
 );
@@ -670,6 +821,15 @@ INSTRUMENTATION_FUNCTION_ATTRIBUTES DWORD NTAPI DetouredResumeThread(HANDLE Thre
 	// read by the thread once it is resumed.
 	FSyncObjectRef Sync = GetSyncObject(Context, ThreadHandle);
 	Sync->SyncReleaseAsSoleOwner(Context, _ReturnAddress(), ThreadHandle, TEXT("ResumeThread"));
+
+	// This handle synchronization of SuspendThread/ResumeThread pairs (e.g. callstack capture)
+	// once the thread is already started.
+	DWORD ThreadId = ::GetThreadId(ThreadHandle);
+	TRefCountPtr<FContext> SuspendedThreadContext = GetContextByThreadId(ThreadId);
+	if (SuspendedThreadContext)
+	{
+		SuspendedThreadContext->ClockBank.Acquire(Context.ClockBank, _ReturnAddress());
+	}
 	DWORD result = TrueResumeThread(ThreadHandle);
 
 	FINISH_WINAPI_INSTRUMENTATION();
@@ -687,6 +847,8 @@ INSTRUMENTATION_FUNCTION_ATTRIBUTES void PopulateDetouredFunctions()
 		return;
 	}
 
+	TrueSuspendThread = ::SuspendThread;
+	DetouredFunctions.Emplace(&(PVOID&)TrueSuspendThread, DetouredSuspendThread);
 	TrueResumeThread = ::ResumeThread;
 	DetouredFunctions.Emplace(&(PVOID&)TrueResumeThread, DetouredResumeThread);
 	TrueWaitForSingleObject = ::WaitForSingleObject;
@@ -1123,13 +1285,16 @@ namespace UE::Sanitizer::RaceDetector::Platform {
 		return ::RtlCaptureStackBackTrace(FrameToSkip, FrameToCapture, Backtrace, 0);
 	}
 
-	INSTRUMENTATION_FUNCTION_ATTRIBUTES static void ThrottleMemoryAccessesDuringShadowReset()
+	INSTRUMENTATION_FUNCTION_ATTRIBUTES void ThrottleMemoryAccessesDuringShadowReset()
 	{
-		// If we're resetting the shadow memory, we delay resuming threads until the reset is done.
-		// because we don't want to mess with the bitmap while it's being cleaned up.
-		while (bIsResettingShadow)
+		if (bIsResettingShadow && !bIsResettingShadowThread)
 		{
-			Sleep(0);
+			// If we're resetting the shadow memory, we delay resuming threads until the reset is done.
+			// because we don't want to mess with the bitmap while it's being cleaned up.
+			while (bIsResettingShadow)
+			{
+				Sleep(0);
+			}
 		}
 	}
 
diff --git a/Engine/Source/Runtime/Core/Tests/Sanitizer/RaceDetectorTests.cpp b/Engine/Source/Runtime/Core/Tests/Sanitizer/RaceDetectorTests.cpp
--- a/Engine/Source/Runtime/Core/Tests/Sanitizer/RaceDetectorTests.cpp
+++ b/Engine/Source/Runtime/Core/Tests/Sanitizer/RaceDetectorTests.cpp
@@ -17,6 +17,9 @@
 #include "Containers/UnrealString.h"
 #include "Sanitizer/RaceDetector.h"
 
+#include "HAL/Runnable.h"
+#include "HAL/RunnableThread.h"
+
 #include <atomic>
 #include <thread>
 
@@ -547,6 +550,139 @@ namespace UE { namespace TasksTests
 		CHECK(!Collector.Contains(Base));
 	}
 
+	// Validates that SuspendThread/ResumeThread act as acquire/release barriers
+	// in the race detector. This is the pattern used by callstack capture code
+	// (WindowsPlatformStackWalk.cpp) where a thread is suspended, its context is
+	// read, and then it is resumed.
+	TEST_CASE_NAMED(FRaceDetectorSuspendThreadBarrierTest, "System::Core::Sanitizer::RaceDetector::SuspendThreadBarrier", "[ApplicationContextMask][EngineFilter]")
+	{
+		CHECK(UE::Sanitizer::RaceDetector::Initialize());
+
+		ToggleRaceDetection(false);
+
+		FThreadRaceCollector Collector;
+
+		ToggleRaceDetection(true);
+
+		int32 X = 0, Y = 0;
+
+		std::atomic<bool> WorkerReady{false};
+		std::atomic<bool> WorkerCanExit{false};
+
+		class FSuspendTestRunnable : public FRunnable
+		{
+		public:
+			int& X;
+			int& Y;
+			std::atomic<bool>& bReady;
+			std::atomic<bool>& bCanExit;
+
+			FSuspendTestRunnable(int& InX, int& InY, std::atomic<bool>& InReady, std::atomic<bool>& InCanExit)
+				: X(InX), Y(InY), bReady(InReady), bCanExit(InCanExit) {}
+
+			virtual uint32 Run() override
+			{
+				// Write X before we get suspended.
+				X = 42;
+				bReady.store(true, std::memory_order_relaxed);
+
+				// Spin until the main thread lets us continue (after Resume).
+				while (!bCanExit.load(std::memory_order_relaxed))
+				{
+					FPlatformProcess::Yield();
+				}
+
+				// Read Y after being resumed - should see main thread's write
+				// with no race because ResumeThread is a release barrier.
+				volatile int32 Sink = Y;
+				(void)Sink;
+
+				return 0;
+			}
+		};
+
+		FSuspendTestRunnable Runnable(X, Y, WorkerReady, WorkerCanExit);
+		FRunnableThread* WorkerThread = FRunnableThread::Create(&Runnable, TEXT("SuspendBarrierTest"));
+
+		// Wait until the worker has written x.
+		while (!WorkerReady.load(std::memory_order_relaxed))
+		{
+			FPlatformProcess::Yield();
+		}
+
+		// Suspend the worker. After this returns, all of the worker's prior
+		// writes (including x = 42) are visible to us without a race.
+		WorkerThread->Suspend(true);
+
+		// Read x while worker is suspended - no race because SuspendThread
+		// establishes an acquire barrier.
+		volatile int32 Sink = X;
+		(void)Sink;
+
+		// Write y while worker is suspended - will be visible to the worker
+		// after ResumeThread (release barrier).
+		Y = 99;
+
+		// Resume the worker.
+		WorkerThread->Suspend(false);
+		WorkerCanExit.store(true, std::memory_order_relaxed);
+
+		WorkerThread->WaitForCompletion();
+		delete WorkerThread;
+
+		ToggleRaceDetection(false);
+
+		// Neither access should be flagged as a race.
+		CHECK(!Collector.Contains(&X));
+		CHECK(!Collector.Contains(&Y));
+	}
+
+	struct FTLSReuseBetweenThreads
+	{
+		int32 Value = 0;
+	};
+
+	static thread_local FTLSReuseBetweenThreads TLSReuseBetweenThreads;
+
+	TEST_CASE_NAMED(FRaceDetectorTLSReuseBetweenThreads, "System::Core::Sanitizer::RaceDetector::TLSReuseBetweenThreads", "[ApplicationContextMask][EngineFilter]")
+	{
+		CHECK(UE::Sanitizer::RaceDetector::Initialize());
+
+		ToggleRaceDetection(false);
+
+		FThreadRaceCollector Collector;
+
+		ToggleRaceDetection(true);
+
+		// Spawning and shutting down thread repetitively is enough to trigger a few false positives
+		// throughout the UE code-base. Also add our own TLS value bump here for completeness.
+		std::atomic<bool> bDone = false;
+		for (int32 i = 0; i < 100; ++i)
+		{
+			bDone.store(false, std::memory_order_relaxed);
+			AsyncThread(
+				[&bDone]()
+				{
+					TLSReuseBetweenThreads.Value++;
+
+					bDone.store(true, std::memory_order_relaxed);
+				});
+			
+			// Wait until the thread is done to increase our chance of
+			// the TLS memory address to be reused for the next spawned thread
+			while (!bDone.load(std::memory_order_relaxed))
+			{
+				FPlatformProcess::Sleep(0);
+			}
+		}
+
+		ToggleRaceDetection(false);
+
+		// if race detector properly detects TLS addresses being released
+		// by a thread, it shouldn't detect any race on the thread_local.
+		CHECK(Collector.IsEmpty());
+	}
+
 	struct FTestStruct
 	{
 		int32 Index;
