[StallDetector]
- Stall detector doesn't make sense for instrumented builds

#Synced-CL 51934851

diff --git a/Engine/Source/Runtime/Core/Public/ProfilingDebugging/StallDetector.h b/Engine/Source/Runtime/Core/Public/ProfilingDebugging/StallDetector.h
--- a/Engine/Source/Runtime/Core/Public/ProfilingDebugging/StallDetector.h
+++ b/Engine/Source/Runtime/Core/Public/ProfilingDebugging/StallDetector.h
@@ -18,7 +18,7 @@
 #include <atomic>
 
 #ifndef STALL_DETECTOR
- #if WITH_EDITOR && ( PLATFORM_WINDOWS || PLATFORM_UNIX ) && !UE_BUILD_SHIPPING && COUNTERSTRACE_ENABLED && !USING_THREAD_SANITISER
+ #if WITH_EDITOR && ( PLATFORM_WINDOWS || PLATFORM_UNIX ) && !UE_BUILD_SHIPPING && COUNTERSTRACE_ENABLED && !USING_THREAD_SANITISER && !USING_INSTRUMENTATION
   #define STALL_DETECTOR 1
  #else
   #define STALL_DETECTOR 0
