- CIS: Add missing include to fix non-unity builds.

#rnx

- Without this include some platforms would run into problems where PLATFORM_EMPTY_BASES is undefined.
- Issue was introduced in CL 51998173

#Synced-CL 52051066

diff --git a/Engine/Source/Runtime/Core/Public/Templates/AlignmentTemplates.h b/Engine/Source/Runtime/Core/Public/Templates/AlignmentTemplates.h
--- a/Engine/Source/Runtime/Core/Public/Templates/AlignmentTemplates.h
+++ b/Engine/Source/Runtime/Core/Public/Templates/AlignmentTemplates.h
@@ -5,6 +5,7 @@
 #include "CoreTypes.h"
 #include "Concepts/Integral.h"
 #include "Concepts/Pointer.h"
+#include "HAL/PlatformMisc.h"
 #include <type_traits>
 
 /**
