Added asserts to FName to prevent them being created with invalid numbers.

diff --git a/Engine/Source/Runtime/Core/Private/UObject/UnrealNames.cpp b/Engine/Source/Runtime/Core/Private/UObject/UnrealNames.cpp
--- a/Engine/Source/Runtime/Core/Private/UObject/UnrealNames.cpp
+++ b/Engine/Source/Runtime/Core/Private/UObject/UnrealNames.cpp
@@ -2915,6 +2915,14 @@ bool FName::IsWithinBounds(FNameEntryId Id)
 	return GetNamePoolPostInit().IsValid(Id);
 }
 
+#if DO_CHECK
+void FName::ReportInvalidInternalNumber(FNameEntryId ComparisonId, int32 InternalNumber)
+{
+	UE_LOG(LogCore, Fatal, TEXT("Trying to construct an invalid FName '%s' with internal number: %d"), *GetNamePool().Resolve(ComparisonId).GetPlainNameString(), InternalNumber);
+	for (;;);
+}
+#endif
+
 /*-----------------------------------------------------------------------------
 	FName implementation.
 -----------------------------------------------------------------------------*/
@@ -3208,6 +3216,14 @@ private:
 			}
 		}
 #endif
+
+#if DO_CHECK
+		if (!FName::IsValidInternalNumber(InternalNumber))
+		{
+			FName::ReportInvalidInternalNumber(Ids.ComparisonId, InternalNumber);
+		}
+#endif
+
 		return MakeWithInternalNumber(Ids, FindType, InternalNumber);
 	}
 
diff --git a/Engine/Source/Runtime/Core/Public/UObject/NameTypes.h b/Engine/Source/Runtime/Core/Public/UObject/NameTypes.h
--- a/Engine/Source/Runtime/Core/Public/UObject/NameTypes.h
+++ b/Engine/Source/Runtime/Core/Public/UObject/NameTypes.h
@@ -837,6 +837,17 @@ public:
 	/** Paranoid sanity check, same as IsValid() */
 	[[nodiscard]] bool IsValidIndexFast() const { return IsValid(); }
 
+	/**
+	 * Checks to see if an internal number is valid, such as those provided to the FName(Name, Number) constructors.
+	 *
+	 * @param	InInternalNumber	The number to check.
+	 *
+	 * @return	true if the number is valid.
+	 */
+	[[nodiscard]] static constexpr bool IsValidInternalNumber(int32 InInternalNumber)
+	{
+		return InInternalNumber >= NAME_NO_NUMBER_INTERNAL;
+	}
 
 	/**
 	 * Checks to see that a given name-like string follows the rules that Unreal requires.
@@ -1294,6 +1305,14 @@ private:
 			// We need to store a new entry in the name table
 			return CreateNumberedName(ComparisonId, DisplayId, InInternalNumber);
 		}
+
+#if DO_CHECK
+		if (!IsValidInternalNumber(InInternalNumber))
+		{
+			ReportInvalidInternalNumber(ComparisonId, InInternalNumber);
+		}
+#endif
+
 		// Otherwise we can just set the index members
 #endif
 		FName Out;
@@ -1363,6 +1382,10 @@ private:
 	}
 #endif
 
+#if DO_CHECK
+	[[noreturn]] CORE_API static void ReportInvalidInternalNumber(FNameEntryId ComparisonId, int32 InternalNumber);
+#endif
+
 public:
 	template <typename T>
 	[[nodiscard]] FORCEINLINE auto UEOpEquals(T Rhs) const -> decltype(FName::Equals(*this, Rhs))
@@ -1849,7 +1872,7 @@ class FDisplayNameEntryId
 public:
 	FDisplayNameEntryId() : FDisplayNameEntryId(FName()) {}
 	explicit FDisplayNameEntryId(FName Name) : FDisplayNameEntryId(Name.GetDisplayIndex(), Name.GetComparisonIndex()) {}
-	FORCEINLINE FName ToName(uint32 InInternalNumber) const { return FName(GetComparisonId(), GetDisplayId(), InInternalNumber); }
+	FORCEINLINE FName ToName(uint32 Number) const { return FName(GetComparisonId(), GetDisplayId(), Number); }
 
 private:
 #if WITH_CASE_PRESERVING_NAME
