Fixed  redefinition of 'SmallBinSizesInternal' error in MallocBinned4

#Synced-CL 52054134

diff --git a/Engine/Source/Runtime/Core/Private/HAL/MallocBinned4.cpp b/Engine/Source/Runtime/Core/Private/HAL/MallocBinned4.cpp
--- a/Engine/Source/Runtime/Core/Private/HAL/MallocBinned4.cpp
+++ b/Engine/Source/Runtime/Core/Private/HAL/MallocBinned4.cpp
@@ -45,7 +45,7 @@ namespace
 // Bin sizes are based around getting the maximum amount of allocations per pool, with as little alignment waste as possible.
 // Bin sizes should be close to even divisors of the system page size, and well distributed.
 // They must be 16-byte aligned as well.
-constexpr uint16 SmallBinSizesInternal[] =
+constexpr uint16 MB4SmallBinSizesInternal[] =
 {
 	16, 32, 48, 64, 80, 96, 112, 128, 144, 160, 176, 192, 208, // +16
 //	Bin		|	Divider | Slack leftover per page (64KB)
@@ -390,10 +390,10 @@ FMallocBinned4::FMallocBinned4()
 	check(!bOnce); // this is now a singleton-like thing and you cannot make multiple copies
 	bOnce = true;
 
-	FMemory::Memcpy(SmallBinSizes, SmallBinSizesInternal);
+	FMemory::Memcpy(SmallBinSizes, MB4SmallBinSizesInternal);
 	
-	static_assert(sizeof(FFreeBlock)  <= SmallBinSizesInternal[0], "sizeof(FFreeBlock)  must be fit in smallest allocation size handled by FMallocBinned4.");
-	static_assert(sizeof(FBundleNode) <= SmallBinSizesInternal[0], "sizeof(FBundleNode) must fit in smallest allocation size handled by FMallocBinned4.");
+	static_assert(sizeof(FFreeBlock)  <= MB4SmallBinSizesInternal[0], "sizeof(FFreeBlock)  must be fit in smallest allocation size handled by FMallocBinned4.");
+	static_assert(sizeof(FBundleNode) <= MB4SmallBinSizesInternal[0], "sizeof(FBundleNode) must fit in smallest allocation size handled by FMallocBinned4.");
 
 	const FGenericPlatformMemoryConstants Constants = FPlatformMemory::GetConstants();
 	GMB4PageSize = Constants.BinnedPageSize;
@@ -405,9 +405,9 @@ FMallocBinned4::FMallocBinned4()
 	checkf(Constants.AddressLimit > GMB4PageSize, TEXT("OS address limit must be greater than the page size")); // Check to catch 32 bit overflow in AddressLimit
 	checkf(GMB4PageSize % UE_MB4_LARGE_ALLOC == 0, TEXT("OS page size must be a multiple of UE_MB4_LARGE_ALLOC"));
 
-	static_assert(SmallBinSizesInternal[UE_MB4_SMALL_POOL_COUNT - 1] == UE_MB4_MAX_SMALL_POOL_SIZE, "UE_MB4_MAX_SMALL_POOL_SIZE must equal the largest bin size");
-	static_assert(UE_ARRAY_COUNT(SmallBinSizesInternal) == UE_MB4_SMALL_POOL_COUNT, "Number of bins in SmallBinSizesInternal must match UE_MB4_SMALL_POOL_COUNT");
-	static_assert(UE_ARRAY_COUNT(SmallBinSizesInternal) <= 256, "Number of bins in SmallBinSizesInternal fit in one byte");
+	static_assert(MB4SmallBinSizesInternal[UE_MB4_SMALL_POOL_COUNT - 1] == UE_MB4_MAX_SMALL_POOL_SIZE, "UE_MB4_MAX_SMALL_POOL_SIZE must equal the largest bin size");
+	static_assert(UE_ARRAY_COUNT(MB4SmallBinSizesInternal) == UE_MB4_SMALL_POOL_COUNT, "Number of bins in MB4SmallBinSizesInternal must match UE_MB4_SMALL_POOL_COUNT");
+	static_assert(UE_ARRAY_COUNT(MB4SmallBinSizesInternal) <= 256, "Number of bins in MB4SmallBinSizesInternal fit in one byte");
 
 	// Init pool tables.
 	for (uint32 Index = 0; Index != UE_MB4_SMALL_POOL_COUNT; ++Index)
