[MallocCrash]
- Suppress benign race on threadid during crash

#Synced-CL 51934851

diff --git a/Engine/Source/Runtime/Core/Public/GenericPlatform/GenericPlatformMallocCrash.h b/Engine/Source/Runtime/Core/Public/GenericPlatform/GenericPlatformMallocCrash.h
--- a/Engine/Source/Runtime/Core/Public/GenericPlatform/GenericPlatformMallocCrash.h
+++ b/Engine/Source/Runtime/Core/Public/GenericPlatform/GenericPlatformMallocCrash.h
@@ -134,7 +134,7 @@ protected:
 	FCriticalSection InternalLock;
 
 	/** The id of the thread that crashed. */
-	int32 CrashedThreadId = 0;
+	TSAN_ATOMIC(int32) CrashedThreadId = 0;
 
 	/** Preallocated memory pool for large allocations. */
 	uint8* LargeMemoryPool = nullptr;
