Add all thread callstacks to Linux crash reports

#rnx
#Synced-CL 52099492

diff --git a/Engine/Source/Runtime/Core/Private/Unix/UnixPlatformCrashContext.cpp b/Engine/Source/Runtime/Core/Private/Unix/UnixPlatformCrashContext.cpp
--- a/Engine/Source/Runtime/Core/Private/Unix/UnixPlatformCrashContext.cpp
+++ b/Engine/Source/Runtime/Core/Private/Unix/UnixPlatformCrashContext.cpp
@@ -28,6 +28,7 @@
 #include "Unix/UnixPlatformRealTimeSignals.h"
 #include "Unix/UnixPlatformRunnableThread.h"
 #include "HAL/ExceptionHandling.h"
+#include "HAL/ThreadManager.h"
 #include "Stats/Stats.h"
 #include "HAL/ThreadHeartBeat.h"
 #include "BuildSettings.h"
@@ -38,6 +39,12 @@
 
 extern CORE_API bool GIsGPUCrashed;
 
+static TAutoConsoleVariable<bool> CVarUnixCaptureAllThreadStacks(
+	TEXT("unix.CaptureAllThreadStacksOnCrash"),
+	true,
+	TEXT("When true, capture callstacks from all threads during crash reporting.\n")
+	TEXT("Set to false to disable if this causes issues such as hangs or recursive crashes."));
+
 FString DescribeSignal(int32 Signal, siginfo_t* Info, ucontext_t *Context)
 {
 	FString ErrorString;
@@ -333,6 +340,96 @@ void FUnixCrashContext::CaptureThreadStackTrace(uint32_t ThreadId)
 	}
 }
 
+// Defined in UnixPlatformStackWalk.cpp - prevents deliberate crash on timeout during crash handling
+extern bool GUnixStackWalkNoCrashOnTimeout;
+
+void FUnixCrashContext::CaptureAllThreadContexts()
+{
+#if PLATFORM_SUPPORTS_ALL_THREAD_BACKTRACES
+	if (!CVarUnixCaptureAllThreadStacks.AsVariable()->GetBool())
+	{
+		return;
+	}
+
+	const uint32 CrashThreadId = FPlatformTLS::GetCurrentThreadId();
+
+	// Reserve upfront to minimize allocations during crash handling.
+	// ~1KB per thread XML entry, 128 threads max = ~128KB.
+	AllThreadContextsString.Reserve(128 * 1024);
+
+	// Prevent recursive crash if a thread is unresponsive during stack capture
+	GUnixStackWalkNoCrashOnTimeout = true;
+
+	FThreadManager::Get().ForEachThreadStackBackTrace(
+		[this, CrashThreadId](uint32 ThreadId, const TCHAR* ThreadName, const TConstArrayView<uint64>& StackTrace)
+		{
+			TArray<FCrashStackFrame> PortableStack;
+			if (ThreadId == CrashThreadId && CallStack.Num() > 0)
+			{
+				// For the crashing thread, use the already-captured CallStack which has
+				// better accuracy from the signal context than a re-capture would
+				PortableStack = CallStack;
+			}
+			else
+			{
+				GetPortableCallStack(StackTrace.GetData(), StackTrace.Num(), PortableStack);
+			}
+
+			// Skip threads where we failed to capture a callstack (e.g. thread already
+			// exited, or the RT signal could not be delivered). This matches Windows
+			// behavior where such threads don't appear with empty stacks.
+			if (PortableStack.Num() == 0 && ThreadId != CrashThreadId)
+			{
+				return true;
+			}
+
+			// Format as XML and append
+			int32 MaxModuleNameLen = 0;
+			for (const FCrashStackFrame& StFrame : PortableStack)
+			{
+				MaxModuleNameLen = FMath::Max(MaxModuleNameLen, StFrame.ModuleName.Len());
+			}
+
+			AllThreadContextsString += TEXT("<Thread>");
+			{
+				AllThreadContextsString += TEXT("<CallStack>");
+				FString CallstackStr;
+				for (const FCrashStackFrame& StFrame : PortableStack)
+				{
+					CallstackStr += FString::Printf(TEXT("%-*s 0x%016llx + %-16llx"), MaxModuleNameLen + 1, *StFrame.ModuleName, StFrame.BaseAddress, StFrame.Offset);
+					CallstackStr += LINE_TERMINATOR;
+				}
+				AppendEscapedXMLString(AllThreadContextsString, *CallstackStr);
+				AllThreadContextsString += TEXT("</CallStack>");
+				AllThreadContextsString += LINE_TERMINATOR;
+			}
+			AllThreadContextsString += FString::Printf(TEXT("<IsCrashed>%s</IsCrashed>"), ThreadId == CrashThreadId ? TEXT("true") : TEXT("false"));
+			AllThreadContextsString += LINE_TERMINATOR;
+			AllThreadContextsString += TEXT("<Registers></Registers>");
+			AllThreadContextsString += LINE_TERMINATOR;
+			AllThreadContextsString += FString::Printf(TEXT("<ThreadID>%d</ThreadID>"), ThreadId);
+			AllThreadContextsString += LINE_TERMINATOR;
+			AllThreadContextsString += TEXT("<ThreadName>");
+			AppendEscapedXMLString(AllThreadContextsString, ThreadName);
+			AllThreadContextsString += TEXT("</ThreadName>");
+			AllThreadContextsString += LINE_TERMINATOR;
+			AllThreadContextsString += TEXT("</Thread>");
+			AllThreadContextsString += LINE_TERMINATOR;
+
+			return true;
+		});
+
+	GUnixStackWalkNoCrashOnTimeout = false;
+#endif
+}
+
+bool FUnixCrashContext::GetPlatformAllThreadContextsString(FString& OutStr) const
+{
+	OutStr = AllThreadContextsString;
+	AllThreadContextsString.Empty();
+	return !OutStr.IsEmpty();
+}
+
 void FUnixCrashContext::GetPortableCallStack(const uint64* StackFrames, int32 NumStackFrames, TArray<FCrashStackFrame>& OutCallStack) const
 {
 	// Update the callstack with offsets from each module
@@ -885,6 +982,7 @@ void DefaultCrashHandler(const FUnixCrashContext & Context)
 
 	// at this point we should already be using malloc crash handler (see PlatformCrashHandler)
 	const_cast<FUnixCrashContext&>(Context).CaptureStackTrace(Context.ErrorFrame);
+	const_cast<FUnixCrashContext&>(Context).CaptureAllThreadContexts();
 	if (GLog)
 	{
 		GLog->Panic();
diff --git a/Engine/Source/Runtime/Core/Private/Unix/UnixPlatformStackWalk.cpp b/Engine/Source/Runtime/Core/Private/Unix/UnixPlatformStackWalk.cpp
--- a/Engine/Source/Runtime/Core/Private/Unix/UnixPlatformStackWalk.cpp
+++ b/Engine/Source/Runtime/Core/Private/Unix/UnixPlatformStackWalk.cpp
@@ -29,6 +29,10 @@ static TAutoConsoleVariable<float> CVarUnixPlatformThreadCallStackMaxWait(
 	60.0f,
 	TEXT("The number of seconds allowed to spin before killing the process, with the assumption the signal handler has hung."));
 
+// When true, thread stack capture timeout should not cause a deliberate crash.
+// Set during crash handling to prevent recursive crashes when gathering all thread stacks.
+bool GUnixStackWalkNoCrashOnTimeout = false;
+
 // Init'ed in UnixPlatformMemory. Once this is tested more we can remove this fallback flag
 extern bool CORE_API GFullCrashCallstack;
 
@@ -799,15 +803,27 @@ UE_AUTORTFM_ALWAYS_OPEN uint32 FUnixPlatformStackWalk::CaptureStackBackTrace( ui
 
 namespace
 {
+	// Shorter timeout for crash handling — the default CVar is 60s which would cause
+	// multi-minute waits if several threads are unresponsive.
+	constexpr double CrashHandlingThreadStackMaxWait = 2.0;
+
 	void WaitForSignalHandlerToFinishOrCrash(ThreadStackUserData& ThreadStack)
 	{
-		double EndWaitTimestamp = FPlatformTime::Seconds() + CVarUnixPlatformThreadCallStackMaxWait.AsVariable()->GetFloat();
+		const double MaxWait = GUnixStackWalkNoCrashOnTimeout
+			? CrashHandlingThreadStackMaxWait
+			: CVarUnixPlatformThreadCallStackMaxWait.AsVariable()->GetFloat();
+		double EndWaitTimestamp = FPlatformTime::Seconds() + MaxWait;
 		double CurrentTimestamp = FPlatformTime::Seconds();
 
 		while (!ThreadStack.bDone)
 		{
 			if (CurrentTimestamp > EndWaitTimestamp)
 			{
+				if (GUnixStackWalkNoCrashOnTimeout)
+				{
+					// During crash handling, skip unresponsive threads rather than causing a recursive crash
+					return;
+				}
 				// We have waited for as long as we should for the signal handler to finish. Assume it has hang and we need to kill our selfs
 				*(int*)0x10 = 0x0;
 			}
@@ -834,6 +850,13 @@ namespace
 		if (syscall(SYS_rt_tgsigqueueinfo, info.si_pid, ThreadId, THREAD_CALLSTACK_GENERATOR, &info) == 0)
 		{
 			WaitForSignalHandlerToFinishOrCrash(ThreadStack);
+
+			if (!ThreadStack.bDone && GUnixStackWalkNoCrashOnTimeout)
+			{
+				// Timed out waiting for signal handler. Explicitly mark as failed so
+				// callers see a deterministic 0 count rather than stale/partial data.
+				ThreadStack.BackTraceCount = 0;
+			}
 		}
 	}
 }
@@ -852,7 +875,19 @@ UE_AUTORTFM_ALWAYS_OPEN void FUnixPlatformStackWalk::ThreadStackWalkAndDump(ANSI
 
 UE_AUTORTFM_ALWAYS_OPEN uint32 FUnixPlatformStackWalk::CaptureThreadStackBackTrace(uint64 ThreadId, uint64* BackTrace, uint32 MaxDepth, void* Context)
 {
-	ThreadStackUserData ThreadBackTrace;
+	// During crash handling (GUnixStackWalkNoCrashOnTimeout == true), use a ring buffer of
+	// static entries so that if the signal times out and we return early, the in-flight signal
+	// handler on the target thread still has valid memory to write to (avoiding use-after-scope
+	// on a stack local). The ring buffer ensures each timed-out capture gets its own slot,
+	// preventing a late-arriving handler from corrupting the current capture's data.
+	static constexpr uint32 RingBufferSize = 8;
+	static ThreadStackUserData StaticRingBuffer[RingBufferSize];
+	static uint32 RingBufferIndex = 0;
+	ThreadStackUserData StackThreadBackTrace;
+
+	ThreadStackUserData& ThreadBackTrace = GUnixStackWalkNoCrashOnTimeout
+		? StaticRingBuffer[RingBufferIndex++ % RingBufferSize]
+		: StackThreadBackTrace;
 	ThreadBackTrace.bCaptureCallStack = false;
 	ThreadBackTrace.CallStackSize     = MaxDepth;
 	ThreadBackTrace.BackTrace         = BackTrace;
diff --git a/Engine/Source/Runtime/Core/Public/HAL/ThreadManager.h b/Engine/Source/Runtime/Core/Public/HAL/ThreadManager.h
--- a/Engine/Source/Runtime/Core/Public/HAL/ThreadManager.h
+++ b/Engine/Source/Runtime/Core/Public/HAL/ThreadManager.h
@@ -9,7 +9,7 @@
 #include "HAL/CriticalSection.h"
 
 #ifndef PLATFORM_SUPPORTS_ALL_THREAD_BACKTRACES
-	#define PLATFORM_SUPPORTS_ALL_THREAD_BACKTRACES (PLATFORM_WINDOWS || PLATFORM_MAC)
+	#define PLATFORM_SUPPORTS_ALL_THREAD_BACKTRACES (PLATFORM_WINDOWS || PLATFORM_MAC || PLATFORM_UNIX)
 #endif
 
 class FRunnableThread;
diff --git a/Engine/Source/Runtime/Core/Public/Unix/UnixPlatformCrashContext.h b/Engine/Source/Runtime/Core/Public/Unix/UnixPlatformCrashContext.h
--- a/Engine/Source/Runtime/Core/Public/Unix/UnixPlatformCrashContext.h
+++ b/Engine/Source/Runtime/Core/Public/Unix/UnixPlatformCrashContext.h
@@ -125,13 +125,25 @@ struct FUnixCrashContext : public FGenericCrashContext
 
 	CORE_API void AddPlatformSpecificProperties() const;
 
+	/**
+	 * Captures callstacks from all threads for inclusion in crash reports.
+	 * Must be called after CaptureStackTrace() so the crashing thread's stack is available.
+	 */
+	CORE_API void CaptureAllThreadContexts();
+
 protected:
+	CORE_API virtual bool GetPlatformAllThreadContextsString(FString& OutStr) const override;
+
 	/**
 	 * Dumps all the data from crash context to the "minidump" report.
 	 *
 	 * @param DiagnosticsPath Path to put the file to
 	 */
 	CORE_API void GenerateReport(const FString & DiagnosticsPath) const;
+
+private:
+	/** Accumulated XML string for all thread contexts, populated by CaptureAllThreadContexts() */
+	mutable FString AllThreadContextsString;
 };
 
 typedef FUnixCrashContext FPlatformCrashContext;
diff --git a/Engine/Source/Runtime/Unix/UnixCommonStartup/Private/UnixCommonStartup.cpp b/Engine/Source/Runtime/Unix/UnixCommonStartup/Private/UnixCommonStartup.cpp
--- a/Engine/Source/Runtime/Unix/UnixCommonStartup/Private/UnixCommonStartup.cpp
+++ b/Engine/Source/Runtime/Unix/UnixCommonStartup/Private/UnixCommonStartup.cpp
@@ -79,6 +79,7 @@ void CommonUnixCrashHandler(const FGenericCrashContext& GenericContext)
 
 	// better than having mutable fields?
 	const_cast< FUnixCrashContext& >(Context).CaptureStackTrace(Context.ErrorFrame);
+	const_cast< FUnixCrashContext& >(Context).CaptureAllThreadContexts();
 	if (GLog)
 	{
 		GLog->Panic();
