[Backout] - CL52102008
#rnx
Original CL Desc
-----------------------------------------------------------------
Forked MB2 into MB4 for future development

diff --git a/Engine/Source/Runtime/Core/Private/HAL/MallocBinned2.cpp b/Engine/Source/Runtime/Core/Private/HAL/MallocBinned2.cpp
--- a/Engine/Source/Runtime/Core/Private/HAL/MallocBinned2.cpp
+++ b/Engine/Source/Runtime/Core/Private/HAL/MallocBinned2.cpp
@@ -30,10 +30,18 @@ PRAGMA_DISABLE_UNSAFE_TYPECAST_WARNINGS
 #	define GBinned2MoveOSFreesOffTimeCriticalThreads	UE_DEFAULT_GBinned2MoveOSFreesOffTimeCriticalThreads
 #endif
 
+namespace
+{
 #if UE_MB2_ALLOCATOR_STATS
 	int64 Binned2HashMemory = 0;
 #endif
 
+#if UE_MB2_ALLOCATOR_STATS_VALIDATION
+	int64 AllocatedSmallPoolMemoryValidation = 0;
+	UE::FPlatformRecursiveMutex ValidationCriticalSection;
+	int32 RecursionCounter = 0;
+#endif
+
 // Bin sizes are based around getting the maximum amount of allocations per pool, with as little alignment waste as possible.
 // Bin sizes should be close to even divisors of the system page size, and well distributed.
 // They must be 16-byte aligned as well.
@@ -85,6 +93,8 @@ constexpr uint16 SmallBinSizesInternal[] =
 
 static uint32 GMB2PageSize = 0;
 
+} //~nameless namespace
+
 uint16 FMallocBinned2::SmallBinSizes[UE_MB2_SMALL_POOL_COUNT] = { 0 };
 FMallocBinned2* FMallocBinned2::MallocBinned2 = nullptr;
 // Mapping of sizes to small table indices
diff --git a/Engine/Source/Runtime/Core/Private/HAL/MallocBinned4.cpp b/Engine/Source/Runtime/Core/Private/HAL/MallocBinned4.cpp
deleted file mode 100644
diff --git a/Engine/Source/Runtime/Core/Public/HAL/MallocBinned2.h b/Engine/Source/Runtime/Core/Public/HAL/MallocBinned2.h
--- a/Engine/Source/Runtime/Core/Public/HAL/MallocBinned2.h
+++ b/Engine/Source/Runtime/Core/Public/HAL/MallocBinned2.h
@@ -37,6 +37,13 @@
 #endif
 
 #define UE_MB2_ALLOCATOR_STATS				UE_MBC_ALLOCATOR_STATS
+#define UE_MB2_ALLOCATOR_STATS_VALIDATION	(UE_MB2_ALLOCATOR_STATS && 0)
+
+#if UE_MB2_ALLOCATOR_STATS_VALIDATION
+	extern int64 AllocatedSmallPoolMemoryValidation;
+	extern UE::FPlatformRecursiveMutex ValidationCriticalSection;
+	extern int32 RecursionCounter;
+#endif
 
 // Canary value used in FFreeBlock
 // A constant value unless we're compiled with fork support in which case there are two values identifying whether the page
@@ -328,4 +335,7 @@ public:
 
 #if UE_ENABLE_INCLUDE_ORDER_DEPRECATED_IN_5_7
 #	include "HAL/CriticalSection.h"
+#	if UE_MB2_ALLOCATOR_STATS_VALIDATION
+#		include "Misc/ScopeLock.h"
+#	endif
 #endif
diff --git a/Engine/Source/Runtime/Core/Public/HAL/MallocBinned4.h b/Engine/Source/Runtime/Core/Public/HAL/MallocBinned4.h
deleted file mode 100644
