[AutoRTFM] Fix up remaining AutoRTFM-enabled functions calling AutoRTFM-disabled functions

These are all compiler errors with the improved compiler analysis.

#Synced-CL 52126916

#Synced-CL 52060281

diff --git a/Engine/Source/Programs/AutoRTFMTests/Private/UECoreTests.cpp b/Engine/Source/Programs/AutoRTFMTests/Private/UECoreTests.cpp
--- a/Engine/Source/Programs/AutoRTFMTests/Private/UECoreTests.cpp
+++ b/Engine/Source/Programs/AutoRTFMTests/Private/UECoreTests.cpp
@@ -908,7 +908,7 @@ TEST_CASE("UECore.TScopeLock_TransactionallySafeCriticalSection")
 
 	SECTION("In Static Local Initializer Called From Open")
 	{
-		struct MyStruct final
+		struct AUTORTFM_ENABLE MyStruct final
 		{
 			FTransactionallySafeCriticalSection CriticalSection;
 		};
diff --git a/Engine/Source/Runtime/AutoRTFM/Private/CompilerInterface.cpp b/Engine/Source/Runtime/AutoRTFM/Private/CompilerInterface.cpp
--- a/Engine/Source/Runtime/AutoRTFM/Private/CompilerInterface.cpp
+++ b/Engine/Source/Runtime/AutoRTFM/Private/CompilerInterface.cpp
@@ -187,6 +187,7 @@ extern "C" AUTORTFM_DISABLE_NO_SANITIZE UE_AUTORTFM_API AUTORTFM_NO_ASAN void au
 // Note: Internal.aem maps this to itself when called in the closed
 extern "C" void* autortfm_lookup_function(void* OriginalFunction, const char* Where) noexcept
 {
+	AutoRTFM::UnreachableIfClosed();  // AEM uses this as the closed variant, but always uninstrumented.
 	return FunctionMapLookup(OriginalFunction, Where);
 }
 
diff --git a/Engine/Source/Runtime/Core/Public/Misc/AutomationTest.h b/Engine/Source/Runtime/Core/Public/Misc/AutomationTest.h
--- a/Engine/Source/Runtime/Core/Public/Misc/AutomationTest.h
+++ b/Engine/Source/Runtime/Core/Public/Misc/AutomationTest.h
@@ -4264,7 +4264,7 @@ public: \
 	};
 
 #define BEGIN_DEFINE_SPEC_PRIVATE( TClass, PrettyName, TFlags, FileName, LineNumber ) \
-	class TClass : public FAutomationSpecBase \
+	class AUTORTFM_DISABLE_IF(!EnumHasAnyFlags(TFlags, EAutomationTestFlags::SupportsAutoRTFM)) TClass : public FAutomationSpecBase \
 	{ \
 	public: \
 		TClass( const FString& InName ) \
@@ -4285,7 +4285,6 @@ public: \
 		virtual int32 GetTestSourceFileLine() const override { return LineNumber; } \
 	protected: \
 		virtual FString GetBeautifiedTestName() const override { return PrettyName; } \
-		AUTORTFM_DISABLE_IF(!EnumHasAnyFlags(TFlags, EAutomationTestFlags::SupportsAutoRTFM)) \
 		virtual void Define() override;
 
 #define REGISTER_SIMPLE_AUTOMATION_TEST_TAGS( TClass, PrettyName, TagsString ) \
diff --git a/Engine/Source/Runtime/Core/Public/Modules/ModuleManager.h b/Engine/Source/Runtime/Core/Public/Modules/ModuleManager.h
--- a/Engine/Source/Runtime/Core/Public/Modules/ModuleManager.h
+++ b/Engine/Source/Runtime/Core/Public/Modules/ModuleManager.h
@@ -947,7 +947,7 @@ namespace UE::Core::Private
 		/** Global registrant object for this module when linked statically */ \
 		static FStaticallyLinkedModuleRegistrant< ModuleImplClass > ModuleRegistrant##ModuleName( TEXT(#ModuleName) ); \
 		/* Forced reference to this function is added by the linker to check that each module uses IMPLEMENT_MODULE */ \
-		extern "C" void IMPLEMENT_MODULE_##ModuleName() { UE_STATIC_ASSERT_WARN(UE::Core::Private::ModuleNameEquals(#ModuleName, UE_MODULE_NAME ), "Module name mismatch (" #ModuleName " != " UE_MODULE_NAME "). Please ensure module name passed to IMPLEMENT_MODULE is " UE_MODULE_NAME " to avoid runtime errors in monolithic builds."); } \
+		extern "C" AUTORTFM_DISABLE void IMPLEMENT_MODULE_##ModuleName() { UE_STATIC_ASSERT_WARN(UE::Core::Private::ModuleNameEquals(#ModuleName, UE_MODULE_NAME ), "Module name mismatch (" #ModuleName " != " UE_MODULE_NAME "). Please ensure module name passed to IMPLEMENT_MODULE is " UE_MODULE_NAME " to avoid runtime errors in monolithic builds."); } \
 		PER_MODULE_BOILERPLATE_ANYLINK(ModuleImplClass, ModuleName)
 
 #else
@@ -965,7 +965,7 @@ namespace UE::Core::Private
 		} \
 		static FModuleInitializerEntry ModuleName##InitializerEntry(TEXT(#ModuleName), Initialize##ModuleName##Module, TEXT(UE_MODULE_NAME)); \
 		/* Forced reference to this function is added by the linker to check that each module uses IMPLEMENT_MODULE */ \
-		extern "C" void IMPLEMENT_MODULE_##ModuleName() { UE_STATIC_ASSERT_WARN(UE::Core::Private::ModuleNameEquals(#ModuleName, UE_MODULE_NAME ), "Module name mismatch (" #ModuleName " != " UE_MODULE_NAME "). Please ensure module name passed to IMPLEMENT_MODULE is " UE_MODULE_NAME " to avoid runtime errors in monolithic builds."); } \
+		extern "C" AUTORTFM_DISABLE void IMPLEMENT_MODULE_##ModuleName() { UE_STATIC_ASSERT_WARN(UE::Core::Private::ModuleNameEquals(#ModuleName, UE_MODULE_NAME ), "Module name mismatch (" #ModuleName " != " UE_MODULE_NAME "). Please ensure module name passed to IMPLEMENT_MODULE is " UE_MODULE_NAME " to avoid runtime errors in monolithic builds."); } \
 		PER_MODULE_BOILERPLATE_ANYLINK(ModuleImplClass, ModuleName)
 
 #endif //IS_MONOLITHIC
diff --git a/Engine/Source/Runtime/Core/Public/Stats/StatsSystemTypes.h b/Engine/Source/Runtime/Core/Public/Stats/StatsSystemTypes.h
--- a/Engine/Source/Runtime/Core/Public/Stats/StatsSystemTypes.h
+++ b/Engine/Source/Runtime/Core/Public/Stats/StatsSystemTypes.h
@@ -1761,7 +1761,7 @@ protected:
 };
 
 template<class TStatData, bool TCompiledIn>
-struct FThreadSafeStaticStatInner : public FThreadSafeStaticStatBase
+struct AUTORTFM_OPEN FThreadSafeStaticStatInner : public FThreadSafeStaticStatBase
 {
 	inline TStatId GetStatId() const
 	{
diff --git a/Engine/Source/Runtime/SlateCore/Private/Test/SlateChildrenTest.cpp b/Engine/Source/Runtime/SlateCore/Private/Test/SlateChildrenTest.cpp
--- a/Engine/Source/Runtime/SlateCore/Private/Test/SlateChildrenTest.cpp
+++ b/Engine/Source/Runtime/SlateCore/Private/Test/SlateChildrenTest.cpp
@@ -1,6 +1,7 @@
 // Copyright Epic Games, Inc. All Rights Reserved.
 
 #include "AutoRTFM.h"
+#include "AutoRTFM/Constants.h"
 #include "CoreMinimal.h"
 #include "Misc/AutomationTest.h"
 
@@ -250,7 +251,8 @@ namespace UE::Slate::Test::Private
 	}
 
 	template<typename Pred>
-	AUTORTFM_INFER void Recursive(SWidget & Widget, Pred & InPred)
+	AUTORTFM_DISABLE_IF(AutoRTFM::IsDisabled(AutoRTFM::CallMode<Pred, SWidget&>))
+	void Recursive(SWidget & Widget, Pred & InPred)
 	{
 		Widget.GetAllChildren()->ForEachWidget([&InPred](SWidget& Widget)
 			{
