Networking
- Removed old deprecated code 5.1/5.2/5.3

#Synced-CL 52054134

diff --git a/Engine/Source/Runtime/Core/Private/Misc/NetworkVersion.cpp b/Engine/Source/Runtime/Core/Private/Misc/NetworkVersion.cpp
--- a/Engine/Source/Runtime/Core/Private/Misc/NetworkVersion.cpp
+++ b/Engine/Source/Runtime/Core/Private/Misc/NetworkVersion.cpp
@@ -97,12 +97,6 @@ uint32 FNetworkVersion::CachedNetworkChecksum			= 0;
 bool FNetworkVersion::bHasCachedReplayChecksum = false;
 uint32 FNetworkVersion::CachedReplayChecksum = 0;
 
-uint32 FNetworkVersion::EngineNetworkProtocolVersion	= FEngineNetworkCustomVersion::LatestVersion;
-uint32 FNetworkVersion::GameNetworkProtocolVersion		= 0;
-
-uint32 FNetworkVersion::EngineCompatibleNetworkProtocolVersion		= FEngineNetworkCustomVersion::ReplayBackwardsCompat;
-uint32 FNetworkVersion::GameCompatibleNetworkProtocolVersion		= 0;
-
 void FNetworkVersion::SetProjectVersion(const TCHAR* InVersion)
 {
 	if (ensureMsgf(InVersion != nullptr && FCString::Strlen(InVersion), TEXT("ProjectVersion used for network version must be a valid string!")))
@@ -120,9 +114,6 @@ void FNetworkVersion::SetGameNetworkProtocolVersion(const uint32 InGameNetworkPr
 {
 	FCustomVersionContainer& NetworkCustomVersions = UE::Net::Private::GetNetworkCustomVersions();
 
-	PRAGMA_DISABLE_DEPRECATION_WARNINGS
-	GameNetworkProtocolVersion = InGameNetworkProtocolVersion;
-	PRAGMA_ENABLE_DEPRECATION_WARNINGS
 	NetworkCustomVersions.SetVersion(FGameNetworkCustomVersion::Guid, InGameNetworkProtocolVersion, TEXT("GameNetworkVersion"));
 	bHasCachedNetworkChecksum = false;
 
@@ -133,9 +124,6 @@ void FNetworkVersion::SetGameCompatibleNetworkProtocolVersion(const uint32 InGam
 {
 	FCustomVersionContainer& CompatibleNetworkCustomVersions = UE::Net::Private::GetCompatibleNetworkCustomVersions();
 
-	PRAGMA_DISABLE_DEPRECATION_WARNINGS
-	GameCompatibleNetworkProtocolVersion = InGameCompatibleNetworkProtocolVersion;
-	PRAGMA_ENABLE_DEPRECATION_WARNINGS
 	CompatibleNetworkCustomVersions.SetVersion(FGameNetworkCustomVersion::Guid, InGameCompatibleNetworkProtocolVersion, TEXT("GameNetworkVersion"));
 	
 	bHasCachedNetworkChecksum = false;
@@ -198,26 +186,6 @@ uint32 FNetworkVersion::GetCompatibleNetworkProtocolVersion(const FGuid& Version
 	return CustomVersion ? CustomVersion->Version : 0;
 }
 
-uint32 FNetworkVersion::GetEngineNetworkProtocolVersion()
-{
-	return GetNetworkProtocolVersion(FEngineNetworkCustomVersion::Guid);
-}
-
-uint32 FNetworkVersion::GetEngineCompatibleNetworkProtocolVersion()
-{
-	return GetCompatibleNetworkProtocolVersion(FEngineNetworkCustomVersion::Guid);
-}
-
-uint32 FNetworkVersion::GetGameNetworkProtocolVersion()
-{
-	return GetNetworkProtocolVersion(FGameNetworkCustomVersion::Guid);
-}
-
-uint32 FNetworkVersion::GetGameCompatibleNetworkProtocolVersion()
-{
-	return GetCompatibleNetworkProtocolVersion(FGameNetworkCustomVersion::Guid);
-}
-
 const FCustomVersionContainer& FNetworkVersion::GetNetworkCustomVersions()
 {
 	return UE::Net::Private::GetNetworkCustomVersions();
diff --git a/Engine/Source/Runtime/Core/Public/Misc/EngineNetworkCustomVersion.h b/Engine/Source/Runtime/Core/Public/Misc/EngineNetworkCustomVersion.h
--- a/Engine/Source/Runtime/Core/Public/Misc/EngineNetworkCustomVersion.h
+++ b/Engine/Source/Runtime/Core/Public/Misc/EngineNetworkCustomVersion.h
@@ -5,46 +5,6 @@
 #include "CoreTypes.h"
 #include "Misc/Guid.h"
 
-enum UE_DEPRECATED(5.1, "Using custom versions instead going forward, see FEngineNetworkCustomVersion") EEngineNetworkVersionHistory
-{
-	HISTORY_INITIAL = 1,
-	HISTORY_REPLAY_BACKWARDS_COMPAT = 2,			// Bump version to get rid of older replays before backwards compat was turned on officially
-	HISTORY_MAX_ACTOR_CHANNELS_CUSTOMIZATION = 3,	// Bump version because serialization of the actor channels changed
-	HISTORY_REPCMD_CHECKSUM_REMOVE_PRINTF = 4,		// Bump version since the way FRepLayoutCmd::CompatibleChecksum was calculated changed due to an optimization
-	HISTORY_NEW_ACTOR_OVERRIDE_LEVEL = 5,			// Bump version since a level reference was added to the new actor information
-	HISTORY_CHANNEL_NAMES = 6,						// Bump version since channel type is now an fname
-	HISTORY_CHANNEL_CLOSE_REASON = 7,				// Bump version to serialize a channel close reason in bunches instead of bDormant
-	HISTORY_ACKS_INCLUDED_IN_HEADER = 8,			// Bump version since acks are now sent as part of the header
-	HISTORY_NETEXPORT_SERIALIZATION = 9,			// Bump version due to serialization change to FNetFieldExport
-	HISTORY_NETEXPORT_SERIALIZE_FIX = 10,			// Bump version to fix net field export name serialization 
-	HISTORY_FAST_ARRAY_DELTA_STRUCT = 11,			// Bump version to allow fast array serialization, delta struct serialization.
-	HISTORY_FIX_ENUM_SERIALIZATION = 12,			// Bump version to fix enum net serialization issues.
-	HISTORY_OPTIONALLY_QUANTIZE_SPAWN_INFO = 13,	// Bump version to conditionally disable quantization for Scale, Location, and Velocity when spawning network actors.
-	HISTORY_JITTER_IN_HEADER = 14,					// Bump version since we added jitter clock time to packet headers and removed remote saturation
-	HISTORY_CLASSNETCACHE_FULLNAME = 15,			// Bump version to use full paths in GetNetFieldExportGroupForClassNetCache
-	HISTORY_REPLAY_DORMANCY = 16,					// Bump version to support dormancy properly in replays
-	HISTORY_ENUM_SERIALIZATION_COMPAT = 17,			// Bump version to include enum bits required for serialization into compat checksums, as well as unify enum and byte property enum serialization
-	HISTORY_SUBOBJECT_OUTER_CHAIN = 18,				// Bump version to support subobject outer chains matching on client and server
-	HISTORY_HITRESULT_INSTANCEHANDLE = 19,			// Bump version to support FHitResult change of Actor to HitObjectHandle. This change was made in CL 14369221 but a net version wasn't added at the time.
-	HISTORY_INTERFACE_PROPERTY_SERIALIZATION = 20,	// Bump version to support net serialization of FInterfaceProperty
-	HISTORY_MONTAGE_PLAY_INST_ID_SERIALIZATION = 21,// Bump version to support net serialization of FGameplayAbilityRepAnimMontage, addition of PlayInstanceId and removal of bForcePlayBit
-	HISTORY_SERIALIZE_DOUBLE_VECTORS_AS_DOUBLES = 22,// Bump version to support net serialization of double vector types
-	HISTORY_PACKED_VECTOR_LWC_SUPPORT = 23,			// Bump version to support quantized LWC FVector net serialization
-	HISTORY_PAWN_REMOTEVIEWPITCH = 24,				// Bump version to support serialization changes to RemoteViewPitch
-	HISTORY_REPMOVE_SERVERFRAME_AND_HANDLE = 25,	// Bump version to support serialization changes to RepMove so we can get the serverframe and physics handle associated with the object
-	HISTORY_21_AND_VIEWPITCH_ONLY_DO_NOT_USE = 26,  // Bump version to support up to history 21 + HISTORY_PAWN_REMOTEVIEWPITCH.  DO NOT USE!!!
-	HISTORY_PLACEHOLDER = 27,                       // Bump version to a placeholder.  This version is the same as HISTORY_REPMOVE_SERVERFRAME_AND_HANDLE
-	HISTORY_RUNTIME_FEATURES_COMPATIBILITY = 28,	// Bump version to add network runtime feature compatibility test to handshake (hello/upgrade) control messages
-	HISTORY_SOFTOBJECTPTR_NETGUIDS = 29,			// Bump version to support replicating SoftObjectPtrs by NetGuid instead of raw strings.
-	HISTORY_SUBOBJECT_DESTROY_FLAG = 30,			// Bump version to support subobject destruction message flags
-	HISTORY_GAMESTATE_REPLCIATED_TIME_AS_DOUBLE = 31,	// Bump version to support AGameStateBase::ReplicatedWorldTimeSeconds as double instead of float.
-	HISTORY_CUSTOMVERION = 32,                      // Bump version to switch to using custom versions
-	// New history items go above here.
-
-	HISTORY_ENGINENETVERSION_PLUS_ONE,
-	HISTORY_ENGINENETVERSION_LATEST = HISTORY_ENGINENETVERSION_PLUS_ONE - 1,
-};
-
 struct FEngineNetworkCustomVersion
 {
 	enum Type
diff --git a/Engine/Source/Runtime/Core/Public/Misc/NetworkGuid.h b/Engine/Source/Runtime/Core/Public/Misc/NetworkGuid.h
--- a/Engine/Source/Runtime/Core/Public/Misc/NetworkGuid.h
+++ b/Engine/Source/Runtime/Core/Public/Misc/NetworkGuid.h
@@ -14,19 +14,12 @@ public:
 
 	union
 	{
-		UE_DEPRECATED(5.3, "Deprecated in favor of ObjectId")
-		uint32 Value;
 		uint64 ObjectId;
 	};
 
 	FNetworkGUID()
 		: ObjectId(0)
 	{ }
-
-	UE_DEPRECATED(5.3, "No longer used")
-	FNetworkGUID(uint32 V)
-		: ObjectId(V)
-	{ }
 		
 	friend bool operator==(const FNetworkGUID& X, const FNetworkGUID& Y)
 	{
@@ -44,18 +37,6 @@ public:
 		return Ar;
 	}
 
-	UE_DEPRECATED(5.3, "No longer used")
-	void BuildFromNetIndex(int32 StaticNetIndex)
-	{
-		ObjectId = ((uint64)StaticNetIndex << 1 | 1);
-	}
-
-	UE_DEPRECATED(5.3, "No longer used")
-	int32 ExtractNetIndex()
-	{
-		return ((int32)ObjectId & 1) ? (int32)ObjectId >> 1 : 0;
-	}
-
 	friend uint32 GetTypeHash(const FNetworkGUID& Guid)
 	{
 		return ::GetTypeHash(Guid.ObjectId);
@@ -102,12 +83,6 @@ public:
 		return ObjectId < Other.ObjectId;
 	}
 
-	UE_DEPRECATED(5.3, "Deprecated in favor of CreateFromIndex")
-	static FNetworkGUID Make(int32 Seed, bool bIsStatic)
-	{
-		return CreateFromIndex((uint64)Seed, bIsStatic);
-	}
-
 	static FNetworkGUID CreateFromIndex(uint64 NetIndex, bool bIsStatic)
 	{
 		check(NetIndex <= (MAX_uint64 >> 1));
diff --git a/Engine/Source/Runtime/Core/Public/Misc/NetworkVersion.h b/Engine/Source/Runtime/Core/Public/Misc/NetworkVersion.h
--- a/Engine/Source/Runtime/Core/Public/Misc/NetworkVersion.h
+++ b/Engine/Source/Runtime/Core/Public/Misc/NetworkVersion.h
@@ -66,16 +66,6 @@ struct FNetworkVersion
 	static CORE_API uint32 GetNetworkCompatibleChangelist();
 	static CORE_API uint32 GetReplayCompatibleChangelist();
 
-	UE_DEPRECATED(5.2, "Please use GetNetworkProtocolVersion instead.")
-	static CORE_API uint32 GetEngineNetworkProtocolVersion();
-	UE_DEPRECATED(5.2, "Please use GetNetworkProtocolVersion instead.")
-	static CORE_API uint32 GetGameNetworkProtocolVersion();
-
-	UE_DEPRECATED(5.2, "Please use GetCompatibleNetworkProtocolVersion instead.")
-	static CORE_API uint32 GetEngineCompatibleNetworkProtocolVersion();
-	UE_DEPRECATED(5.2, "Please use GetCompatibleNetworkProtocolVersion instead.")
-	static CORE_API uint32 GetGameCompatibleNetworkProtocolVersion();
-
 	static CORE_API uint32 GetNetworkProtocolVersion(const FGuid& VersionGuid);
 	static CORE_API uint32 GetCompatibleNetworkProtocolVersion(const FGuid& VersionGuid);
 
@@ -159,14 +149,4 @@ protected:
 
 	static CORE_API bool		bHasCachedReplayChecksum;
 	static CORE_API uint32	CachedReplayChecksum;
-
-	UE_DEPRECATED(5.2, "Now storing this value in NetworkCustomVersions, do not use directly.")
-	static CORE_API uint32	EngineNetworkProtocolVersion;
-	UE_DEPRECATED(5.2, "Now storing this value in NetworkCustomVersions, do not use directly.")
-	static CORE_API uint32	GameNetworkProtocolVersion;
-
-	UE_DEPRECATED(5.2, "Now storing this value in CompatibleNetworkCustomVersions, do not use directly.")
-	static CORE_API uint32	EngineCompatibleNetworkProtocolVersion;
-	UE_DEPRECATED(5.2, "Now storing this value in CompatibleNetworkCustomVersions, do not use directly.")
-	static CORE_API uint32	GameCompatibleNetworkProtocolVersion;
 };
diff --git a/Engine/Source/Runtime/Engine/Private/ReplayTypes.cpp b/Engine/Source/Runtime/Engine/Private/ReplayTypes.cpp
--- a/Engine/Source/Runtime/Engine/Private/ReplayTypes.cpp
+++ b/Engine/Source/Runtime/Engine/Private/ReplayTypes.cpp
@@ -43,8 +43,8 @@ FNetworkDemoHeader::FNetworkDemoHeader() :
 void FNetworkDemoHeader::SetDefaultNetworkVersions()
 {
 	NetworkChecksum = FNetworkVersion::GetLocalNetworkVersion();
-	EngineNetworkProtocolVersion = FNetworkVersion::GetEngineNetworkProtocolVersion();
-	GameNetworkProtocolVersion = FNetworkVersion::GetGameNetworkProtocolVersion();
+	EngineNetworkProtocolVersion = FNetworkVersion::GetNetworkProtocolVersion(FEngineNetworkCustomVersion::Guid);
+	GameNetworkProtocolVersion = FNetworkVersion::GetNetworkProtocolVersion(FGameNetworkCustomVersion::Guid);
 
 	CustomVersions = FNetworkVersion::GetNetworkCustomVersions();
 	CustomVersions.SetVersion(FReplayCustomVersion::Guid, FReplayCustomVersion::LatestVersion, TEXT("Replay"));
@@ -191,15 +191,15 @@ FArchive& operator<<(FArchive& Ar, FNetworkDemoHeader& Header)
 			// Fix up for LWC compatibility.
 			// Vectors were using operator<< to serialize in some network cases (instead of
 			// the NetSerialize function), but this operator uses EUnrealEngineObjectUE5Version
-			// and not EEngineNetworkVersionHistory for versioning. Therefore, pre-LWC replays
+			// and not FEngineNetworkCustomVersion for versioning. Therefore, pre-LWC replays
 			// that did not save the EUnrealEngineObjectUE5Version will try to read doubles
 			// from these vectors instead of floats.
 			//
-			// However, EEngineNetworkVersionHistory::HISTORY_SERIALIZE_DOUBLE_VECTORS_AS_DOUBLES was
+			// However, FEngineNetworkCustomVersion::SerializeDoubleVectorsAsDoubles was
 			// added around the same time as EUnrealEngineObjectUE5Version::LARGE_WORLD_COORDINATES,
 			// so we use the network version as an approximation the package version to allow
 			// most pre-LWC replays to play back correctly. The compromise is that any replays recorded
-			// between when HISTORY_SERIALIZE_DOUBLE_VECTORS_AS_DOUBLES and LARGE_WORLD_COORDINATES
+			// between when FEngineNetworkCustomVersion::SerializeDoubleVectorsAsDoubles and LARGE_WORLD_COORDINATES
 			// were added won't play back correctly since they didn't store accurate version information.
 			if (Ar.EngineNetVer() < FEngineNetworkCustomVersion::SerializeDoubleVectorsAsDoubles)
 			{
