Fixed MB4 compilation

#Synced-CL 51881412

diff --git a/Engine/Source/Runtime/Core/Private/HAL/MallocBinned4.cpp b/Engine/Source/Runtime/Core/Private/HAL/MallocBinned4.cpp
--- a/Engine/Source/Runtime/Core/Private/HAL/MallocBinned4.cpp
+++ b/Engine/Source/Runtime/Core/Private/HAL/MallocBinned4.cpp
@@ -946,7 +946,7 @@ void FMallocBinned4::UpdateStats()
 	CSV_CUSTOM_STAT(FMemory, AllocatorCachedSlackMB, (int32)(CachedOSPageAllocator.GetCachedFreeTotal() / (1024 * 1024)), ECsvCustomStatOp::Set);
 	CSV_CUSTOM_STAT(FMemory, AllocatorImmediatelyFreeableCachedSlackMB, (int32)(CachedOSPageAllocator.GetCachedImmediatelyFreeable() / (1024 * 1024)), ECsvCustomStatOp::Set);
 
-	UpdateStatsCommon(*this);
+	FMallocBinnedCommonUtils::UpdateStatsCommon(*this);
 	CachedOSPageAllocator.UpdateStats();
 	FScopedVirtualMallocTimer::UpdateStats();
 }
