[RaceDetector]
- Improve performance and memory usage by using a linear allocator for sync objects that dramatically reduce virtual allocation count

#Synced-CL 52099373

diff --git a/Engine/Source/Runtime/Core/Internal/Sanitizer/RaceDetectorPlatform.h b/Engine/Source/Runtime/Core/Internal/Sanitizer/RaceDetectorPlatform.h
--- a/Engine/Source/Runtime/Core/Internal/Sanitizer/RaceDetectorPlatform.h
+++ b/Engine/Source/Runtime/Core/Internal/Sanitizer/RaceDetectorPlatform.h
@@ -80,6 +80,17 @@ namespace UE::Sanitizer::RaceDetector {
 		// The function generates an interprocessor interrupt (IPI) to all processors that are part of the current process affinity.
 		// It guarantees the visibility of write operations performed on one processor to the other processors.
 		void AsymmetricThreadFenceHeavy();
+
+		// Linear allocator for sync objects — reserves virtual address space, commits in chunks,
+		// decommits the entire region in a single syscall on shadow reset.
+		void  InitSyncObjectAllocator(SIZE_T ReservedSize, SIZE_T CommitGranularity);
+		void  ShutdownSyncObjectAllocator();
+		void* SyncObjectAllocatorAlloc(SIZE_T Size);
+		void  SyncObjectAllocatorReset();
+		void  SyncObjectAllocatorAddRef();
+		void  SyncObjectAllocatorRelease();
+		int32 SyncObjectAllocatorGetRefCount();
+		SIZE_T SyncObjectAllocatorGetCommittedSize();
 	}
 }
 
diff --git a/Engine/Source/Runtime/Core/Internal/Sanitizer/RaceDetectorTypes.h b/Engine/Source/Runtime/Core/Internal/Sanitizer/RaceDetectorTypes.h
--- a/Engine/Source/Runtime/Core/Internal/Sanitizer/RaceDetectorTypes.h
+++ b/Engine/Source/Runtime/Core/Internal/Sanitizer/RaceDetectorTypes.h
@@ -760,16 +760,14 @@ namespace UE::Sanitizer::RaceDetector {
 	public:
 		INSTRUMENTATION_FUNCTION_ATTRIBUTES void* operator new(SIZE_T Size)
 		{
-			// Only count SyncObject that have been allocated separately to 
-			// avoid counting the one embedded in the SyncObjectBank.
 			FPlatformAtomics::InterlockedIncrement(&ObjectCount);
-			return FInstrumentationSafeWinAllocator::Alloc(Size);
+			return Platform::SyncObjectAllocatorAlloc(Size);
 		}
 
 		INSTRUMENTATION_FUNCTION_ATTRIBUTES void operator delete(void* Ptr)
 		{
+			// No individual memory free — the linear allocator decommits everything on reset.
 			FPlatformAtomics::InterlockedDecrement(&ObjectCount);
-			FInstrumentationSafeWinAllocator::Free(Ptr);
 		}
 
 		template <typename AtomicOpType>
@@ -894,25 +892,30 @@ namespace UE::Sanitizer::RaceDetector {
 			return ObjectCount;
 		}
 
+		INSTRUMENTATION_FUNCTION_ATTRIBUTES static void ResetObjectCount()
+		{
+			ObjectCount = 0;
+		}
+
 	private:
 		FRWSpinLock Lock;
 		FClockBank ClockBank;
 		static volatile int64 ObjectCount;
 	};
 
-	// One SyncObjectBank per 64-bit aligned address
+	// One SyncObjectBank per 64-bit aligned address.
+	// Lives in the linear allocator — never individually freed.
 	struct FSyncObjectBank
 	{
-		SAFE_OPERATOR_NEW_DELETE()
-
-		INSTRUMENTATION_FUNCTION_ATTRIBUTES FSyncObjectBank()
+		INSTRUMENTATION_FUNCTION_ATTRIBUTES void* operator new(SIZE_T Size)
 		{
 			FPlatformAtomics::InterlockedIncrement(&ObjectCount);
+			return Platform::SyncObjectAllocatorAlloc(Size);
 		}
 
-		INSTRUMENTATION_FUNCTION_ATTRIBUTES ~FSyncObjectBank()
+		INSTRUMENTATION_FUNCTION_ATTRIBUTES void operator delete(void* Ptr)
 		{
-			Reset();
+			// No individual memory free — the linear allocator decommits everything on reset.
 			FPlatformAtomics::InterlockedDecrement(&ObjectCount);
 		}
 
@@ -921,14 +924,13 @@ namespace UE::Sanitizer::RaceDetector {
 			return ObjectCount;
 		}
 
-		INSTRUMENTATION_FUNCTION_ATTRIBUTES FSyncObject* GetSyncObject(int32 Index)
+		INSTRUMENTATION_FUNCTION_ATTRIBUTES static void ResetObjectCount()
 		{
-			if (Index == 0)
-			{
-				return &EmbeddedObject;
-			}
+			ObjectCount = 0;
+		}
 
-			Index--;
+		INSTRUMENTATION_FUNCTION_ATTRIBUTES FSyncObject* GetSyncObject(int32 Index)
+		{
 			if (SyncObjects[Index] == nullptr)
 			{
 				FSyncObject* SyncObject = new FSyncObject();
@@ -941,82 +943,36 @@ namespace UE::Sanitizer::RaceDetector {
 			return SyncObjects[Index];
 		}
 
-		// We maintain a linked list of clock banks for recycling purpose
-		FSyncObjectBank* Next = 0;
-
-		INSTRUMENTATION_FUNCTION_ATTRIBUTES int32 AddRef()
-		{
-			return FPlatformAtomics::InterlockedIncrement(&RefCount);
-		}
-
-		INSTRUMENTATION_FUNCTION_ATTRIBUTES int32 Release()
-		{
-			int32 NewRefCount = FPlatformAtomics::InterlockedDecrement(&RefCount);
-			if (NewRefCount == 0)
-			{
-				ReleaseLastRef();
-			}
-			return NewRefCount;
-		}
-
-		INSTRUMENTATION_FUNCTION_ATTRIBUTES int32 GetRefCount() const
-		{
-			return RefCount;
-		}
-
-		INSTRUMENTATION_FUNCTION_ATTRIBUTES void Reset()
-		{
-			for (int32 Index = 0; Index < 7; ++Index)
-			{
-				if (SyncObjects[Index])
-				{
-					delete SyncObjects[Index];
-					SyncObjects[Index] = nullptr;
-				}
-			}
-		}
-
 		static volatile int64 ObjectCount;
 	private:
-		INSTRUMENTATION_FUNCTION_ATTRIBUTES void ReleaseLastRef();
-
-		volatile int32 RefCount = 0;
-
-		// Save space by allocating the first sync object as part of the bank itself.
-		// Most of the time the sync object will be at offset 0.
-		// The safe allocator uses virtual memory with 4KB pages so this first entry
-		// is completely free.
-		FSyncObject  EmbeddedObject;
-
-		// Contains optional sync object for each unaligned bytes of the 64-bit.
-		FSyncObject* SyncObjects[7] = { 0 };
+		// One sync object per byte offset in the 8-byte aligned address.
+		// Allocated on demand from the linear allocator.
+		FSyncObject* SyncObjects[8] = { 0 };
 	};
 
 	struct FSyncObjectRef
 	{
 		INSTRUMENTATION_FUNCTION_ATTRIBUTES FORCEINLINE FSyncObjectRef(FSyncObjectBank* InBank, FSyncObject* InSyncObject)
-			: Bank(InBank)
-			, Object(InSyncObject)
+			: Object(InSyncObject)
 		{
-			AddRef();
+			Platform::SyncObjectAllocatorAddRef();
 		}
 
 		INSTRUMENTATION_FUNCTION_ATTRIBUTES FORCEINLINE ~FSyncObjectRef()
 		{
-			Release();
+			if (Object)
+			{
+				Platform::SyncObjectAllocatorRelease();
+			}
 		}
 
-		// Not needed for now
 		FSyncObjectRef(const FSyncObjectRef&) = delete;
 		FSyncObjectRef& operator=(const FSyncObjectRef&) = delete;
 		FSyncObjectRef& operator=(FSyncObjectRef&&) = delete;
 
 		INSTRUMENTATION_FUNCTION_ATTRIBUTES FSyncObjectRef(FSyncObjectRef&& Other)
+			: Object(Other.Object)
 		{
-			Bank = Other.Bank;
-			Object = Other.Object;
-
-			Other.Bank = nullptr;
 			Other.Object = nullptr;
 		}
 
@@ -1025,20 +981,6 @@ namespace UE::Sanitizer::RaceDetector {
 			return Object;
 		}
 	private:
-		INSTRUMENTATION_FUNCTION_ATTRIBUTES FORCEINLINE void AddRef()
-		{
-			Bank->AddRef();
-		}
-
-		INSTRUMENTATION_FUNCTION_ATTRIBUTES FORCEINLINE void Release()
-		{
-			if (Bank)
-			{
-				Bank->Release();
-			}
-		}
-
-		FSyncObjectBank* Bank;
 		FSyncObject* Object;
 	};
 
diff --git a/Engine/Source/Runtime/Core/Private/Sanitizer/RaceDetector.cpp b/Engine/Source/Runtime/Core/Private/Sanitizer/RaceDetector.cpp
--- a/Engine/Source/Runtime/Core/Private/Sanitizer/RaceDetector.cpp
+++ b/Engine/Source/Runtime/Core/Private/Sanitizer/RaceDetector.cpp
@@ -139,12 +139,6 @@ namespace UE::Sanitizer::RaceDetector {
 
 	static UE::FRecursiveMutex Mutex;
 
-	// Since the shadow space can be unmapped in a single shot loosing all our pointers,
-	// we need to keep a list of all the banks we've allocated so we can free them.
-	volatile FSyncObjectBank* SyncObjectBankHead = nullptr;
-
-	volatile FSyncObjectBank* FreeObjectBankHead = nullptr;
-
 	FRWSpinLock RaceReportsLock;
 	TArray<FString> RaceReports;
 
@@ -189,8 +183,7 @@ namespace UE::Sanitizer::RaceDetector {
 
 	INSTRUMENTATION_FUNCTION_ATTRIBUTES uint64 GetSyncObjectMemoryUsage()
 	{
-		return FSyncObject::GetObjectCount() * Align(sizeof(FSyncObject), Platform::GetPageSize()) +
-			FSyncObjectBank::GetObjectCount() * Align(sizeof(FSyncObjectBank), Platform::GetPageSize());
+		return Platform::SyncObjectAllocatorGetCommittedSize();
 	}
 
 	INSTRUMENTATION_FUNCTION_ATTRIBUTES uint64 GetHistoryChunkMemoryUsage()
@@ -392,60 +385,6 @@ namespace UE::Sanitizer::RaceDetector {
 		return (FShadowClockBankSlot*)ShadowMemoryAddress;
 	}
 
-	INSTRUMENTATION_FUNCTION_ATTRIBUTES void PushSyncObjectBank(FSyncObjectBank* SyncObjectBank)
-	{
-		while (true)
-		{
-			FSyncObjectBank* LocalSyncBankHead = (FSyncObjectBank*)SyncObjectBankHead;
-			SyncObjectBank->Next = LocalSyncBankHead;
-
-			if (FPlatformAtomics::InterlockedCompareExchangePointer((void**)&SyncObjectBankHead, SyncObjectBank, LocalSyncBankHead) == LocalSyncBankHead)
-			{
-				break;
-			}
-		}
-	}
-
-	INSTRUMENTATION_FUNCTION_ATTRIBUTES void PushFreeObjectBank(FSyncObjectBank* SyncObjectBank)
-	{
-		RACEDETECTOR_CHECK(SyncObjectBank->GetRefCount() == 0);
-
-		// Cleanup all non-aligned objects, those are unlikely to be needed when the
-		// sync bank is reused for a different atomic address.
-		SyncObjectBank->Reset();
-
-		while (true)
-		{
-			FSyncObjectBank* LocalSyncBankHead = (FSyncObjectBank*)FreeObjectBankHead;
-			SyncObjectBank->Next = LocalSyncBankHead;
-
-			if (FPlatformAtomics::InterlockedCompareExchangePointer((void**)&FreeObjectBankHead, SyncObjectBank, LocalSyncBankHead) == LocalSyncBankHead)
-			{
-				FPlatformAtomics::InterlockedDecrement(&FSyncObjectBank::ObjectCount);
-				break;
-			}
-		}
-	}
-
-	INSTRUMENTATION_FUNCTION_ATTRIBUTES FSyncObjectBank* PopFreeObjectBank()
-	{
-		while (true)
-		{
-			FSyncObjectBank* LocalSyncBankHead = (FSyncObjectBank*)FreeObjectBankHead;
-			if (LocalSyncBankHead == nullptr)
-			{
-				return nullptr;
-			}
-			if (FPlatformAtomics::InterlockedCompareExchangePointer((void**)&FreeObjectBankHead, LocalSyncBankHead->Next, LocalSyncBankHead) == LocalSyncBankHead)
-			{
-				FPlatformAtomics::InterlockedIncrement(&FSyncObjectBank::ObjectCount);
-				LocalSyncBankHead->Next = nullptr;
-				RACEDETECTOR_CHECK(LocalSyncBankHead->GetRefCount() == 0);
-				return LocalSyncBankHead;
-			}
-		}
-	}
-
 	INSTRUMENTATION_FUNCTION_ATTRIBUTES FSyncObjectRef GetSyncObject(FContext& Context, void* SyncAddr)
 	{
 		FShadowClockBankSlot* ClockBankSlot = GetShadowClockBank((UPTRINT)SyncAddr);
@@ -453,43 +392,28 @@ namespace UE::Sanitizer::RaceDetector {
 		RACEDETECTOR_CHECK(Context.BankHazard == nullptr);
 
 		// This can race with ResetShadow so use Hazard Pointer mechanism to avoid
-		// the object bank from being deleted while we're still adding a refcount to it.
+		// the object bank from being decommitted while we're creating an FSyncObjectRef to it.
 		while (true)
 		{
 			FSyncObjectBank* LocalBank = ClockBankSlot->SyncObjectBank;
 			if (LocalBank == nullptr)
 			{
-				FSyncObjectBank* NewObjectBank = PopFreeObjectBank();
-				if (NewObjectBank == nullptr)
-				{
-					NewObjectBank = new FSyncObjectBank();
-				}
-
-				// When publishing in shadow memory, we need to make sure the refcount is not 0, since another thread could try to use
-				// the newly published bank from the shadow, perform an AddRef and then release it and we'd end up with
-				// an invalid NewObjectBank before even starting to use it.
-				NewObjectBank->AddRef();
+				FSyncObjectBank* NewObjectBank = new FSyncObjectBank();
 
 				FSyncObjectBank* OldObjectBank = (FSyncObjectBank*)FPlatformAtomics::InterlockedCompareExchangePointer((void**)&ClockBankSlot->SyncObjectBank, NewObjectBank, nullptr);
 				if (OldObjectBank == nullptr)
 				{
+					// We published successfully. No hazard pointer handling needed since
+					// we own the bank and nobody else can see it yet.
 					LocalBank = NewObjectBank;
-
-					// Take our refcount now so that we don't have to do any hazard pointer handling.
 					FSyncObjectRef Result = FSyncObjectRef(LocalBank, LocalBank->GetSyncObject(Index));
-
-					// Publish the object in the linked-list
-					// The creation ref-count is now owned by the linked-list
-					PushSyncObjectBank(LocalBank);
-
 					Context.BankHazard = nullptr;
 					return MoveTemp(Result);
 				}
 
+				// Another thread won the race — our bank is wasted but this is very rare.
+				// The linear allocator will reclaim it on the next reset.
 				LocalBank = OldObjectBank;
-
-				// Should not happen very often, and at most once per thread taking part of the race
-				NewObjectBank->Release();
 			}
 
 			Context.BankHazard = LocalBank;
@@ -500,7 +424,6 @@ namespace UE::Sanitizer::RaceDetector {
 			if (LocalBank == ClockBankSlot->SyncObjectBank)
 			{
 				FSyncObjectRef Result = FSyncObjectRef(LocalBank, LocalBank->GetSyncObject(Index));
-				RACEDETECTOR_CHECK(LocalBank->GetRefCount() >= 1);
 				Context.BankHazard = nullptr;
 				return MoveTemp(Result);
 			}
@@ -1081,12 +1004,6 @@ namespace UE::Sanitizer::RaceDetector {
 		return Hazards;
 	}
 	
-	INSTRUMENTATION_FUNCTION_ATTRIBUTES void FSyncObjectBank::ReleaseLastRef()
-	{
-		// Releasing the last ref will put it back in the free list, freeing it could cause a crash since it is part of a linked-list
-		// and another thread might need to dereference the Next pointer.
-		PushFreeObjectBank(this);
-	}
 
 	INSTRUMENTATION_FUNCTION_ATTRIBUTES void ResetShadow()
 	{
@@ -1114,36 +1031,34 @@ namespace UE::Sanitizer::RaceDetector {
 			UnhookInstrumentation();
 		}
 
-		// In case there are some race while unmapping we will loop until
-		// both collections are empty and coherent to together to avoid false positives.
-		do
-		{
-			// This need to synchronize with GetSyncObject so null the list before resetting shadow memory
-			FSyncObjectBank* ClockBank = (FSyncObjectBank*)FPlatformAtomics::InterlockedExchangePtr((void**)&SyncObjectBankHead, nullptr);
-			
-			ResetShadowMemory();
+		ResetShadowMemory();
 
+		// Wait for all hazard pointers to clear (threads finishing GetSyncObject).
+		// Since instrumentation is stopped, no new hazards will be created.
+		while (true)
+		{
 			TSet<FSyncObjectBank*> Hazards = GatherBankHazards();
-
-			// Garbage collect what we can and put back what we can't.
-			while (ClockBank)
+			if (Hazards.Num() == 0)
 			{
-				FSyncObjectBank* ToDelete = ClockBank;
-				ClockBank = ClockBank->Next;
-
-				if (Hazards.Contains(ToDelete))
-				{
-					// Put it back in the list for next try.
-					PushSyncObjectBank(ToDelete);
-					FPlatformMisc::LowLevelOutputDebugStringf(TEXT("[%d] %p was in hazard list\n"), Platform::GetCurrentThreadId(), ToDelete);
-				}
-				else
-				{
-					ToDelete->Release();
-				}
+				break;
 			}
+			FPlatformMisc::LowLevelOutputDebugStringf(TEXT("[%d] Waiting for %d hazard pointers to clear\n"), Platform::GetCurrentThreadId(), Hazards.Num());
+			FPlatformProcess::Yield();
+		}
 
-		} while (Platform::HasShadowMemoryMapped());
+		// Wait for all outstanding FSyncObjectRef instances to be released.
+		// These are very short-lived (stack-scoped in InstrumentAtomicAccess),
+		// so this should drain almost immediately after instrumentation is stopped.
+		while (Platform::SyncObjectAllocatorGetRefCount() > 0)
+		{
+			FPlatformMisc::LowLevelOutputDebugStringf(TEXT("[%d] Waiting for %d sync object refs to drain\n"), Platform::GetCurrentThreadId(), Platform::SyncObjectAllocatorGetRefCount());
+			FPlatformProcess::Yield();
+		}
+
+		// Decommit all sync object memory in a single syscall
+		Platform::SyncObjectAllocatorReset();
+		FSyncObject::ResetObjectCount();
+		FSyncObjectBank::ResetObjectCount();
 
 		FPlatformMisc::LowLevelOutputDebugStringf(TEXT("[%d] Reset Shadow Memory Ended\n"), Platform::GetCurrentThreadId());
 		
@@ -1499,7 +1414,11 @@ namespace UE::Sanitizer::RaceDetector {
 		}
 
 		Platform::InitShadowMemory();
-		
+
+		// Reserve virtual address space for the sync object linear allocator.
+		// 64GB reserved (only committed on demand in 2MB chunks).
+		Platform::InitSyncObjectAllocator(64ull * 1024 * 1024 * 1024, 2 * 1024 * 1024);
+
 		ShadowMemoryBase = Platform::GetShadowMemoryBase();
 		ShadowMemorySize = Platform::GetShadowMemorySize();
 		ShadowMemoryEnd = ShadowMemoryBase + ShadowMemorySize;
@@ -1584,6 +1503,8 @@ namespace UE::Sanitizer::RaceDetector {
 
 		ResetShadow();
 
+		Platform::ShutdownSyncObjectAllocator();
+
 		UnhookAlwaysOnInstrumentation();
 
 		CleanupHotPatchFunctions();
diff --git a/Engine/Source/Runtime/Core/Private/Sanitizer/RaceDetectorWindows.cpp b/Engine/Source/Runtime/Core/Private/Sanitizer/RaceDetectorWindows.cpp
--- a/Engine/Source/Runtime/Core/Private/Sanitizer/RaceDetectorWindows.cpp
+++ b/Engine/Source/Runtime/Core/Private/Sanitizer/RaceDetectorWindows.cpp
@@ -1129,6 +1129,106 @@ namespace UE::Sanitizer::RaceDetector::Platform {
 		return false;
 	}
 
+	// ------------------------------------------------------------------------------
+	// Linear allocator for FSyncObject and FSyncObjectBank.
+	// Reserves a large virtual address range upfront, commits pages in chunks on demand,
+	// and decommits the entire region in a single syscall on shadow reset.
+	// This replaces per-object VirtualAlloc/VirtualFree which is extremely slow at scale.
+	// ------------------------------------------------------------------------------
+	static uint8*    SyncAllocBase = nullptr;
+	static SIZE_T    SyncAllocReservedSize = 0;
+	static volatile int64 SyncAllocCommittedSize = 0;
+	static volatile int64 SyncAllocOffset = 0;
+	static SIZE_T    SyncAllocCommitGranularity = 0;
+	static volatile int32 SyncAllocGlobalRefCount = 0;
+
+	INSTRUMENTATION_FUNCTION_ATTRIBUTES static void SyncAllocEnsureCommitted(SIZE_T RequiredSize)
+	{
+		while ((SIZE_T)SyncAllocCommittedSize < RequiredSize)
+		{
+			int64 OldCommitted = SyncAllocCommittedSize;
+			if ((SIZE_T)OldCommitted >= RequiredSize)
+			{
+				break;
+			}
+
+			int64 NewCommitted = (int64)Align(RequiredSize, SyncAllocCommitGranularity);
+			NewCommitted = FMath::Min(NewCommitted, (int64)SyncAllocReservedSize);
+
+			// Commit the pages BEFORE advancing the high-water mark.
+			// VirtualAlloc(MEM_COMMIT) is idempotent on already-committed pages,
+			// so overlapping commits from concurrent threads are safe.
+			void* Result = VirtualAlloc(SyncAllocBase + OldCommitted, (SIZE_T)(NewCommitted - OldCommitted), MEM_COMMIT, PAGE_READWRITE);
+			checkf(Result != nullptr, TEXT("SyncObjectAllocator: Failed to commit %llu bytes"), NewCommitted - OldCommitted);
+
+			// Advance the high-water mark. If another thread already advanced past us, that's fine —
+			// the pages are committed either way. We loop and re-check.
+			FPlatformAtomics::InterlockedCompareExchange(&SyncAllocCommittedSize, NewCommitted, OldCommitted);
+		}
+	}
+
+	INSTRUMENTATION_FUNCTION_ATTRIBUTES void InitSyncObjectAllocator(SIZE_T ReservedSize, SIZE_T CommitGranularity)
+	{
+		SyncAllocReservedSize = ReservedSize;
+		SyncAllocCommitGranularity = CommitGranularity;
+		SyncAllocBase = (uint8*)VirtualAlloc(NULL, ReservedSize, MEM_RESERVE, PAGE_NOACCESS);
+		checkf(SyncAllocBase != nullptr, TEXT("SyncObjectAllocator: Failed to reserve %llu bytes of virtual address space"), ReservedSize);
+	}
+
+	INSTRUMENTATION_FUNCTION_ATTRIBUTES void ShutdownSyncObjectAllocator()
+	{
+		if (SyncAllocBase)
+		{
+			VirtualFree(SyncAllocBase, 0, MEM_RELEASE);
+			SyncAllocBase = nullptr;
+			SyncAllocCommittedSize = 0;
+			SyncAllocOffset = 0;
+		}
+	}
+
+	INSTRUMENTATION_FUNCTION_ATTRIBUTES void* SyncObjectAllocatorAlloc(SIZE_T Size)
+	{
+		constexpr SIZE_T CacheLineSize = 64;
+		const SIZE_T AlignedSize = Align(Size, CacheLineSize);
+		const int64 OldOffset = FPlatformAtomics::InterlockedAdd(&SyncAllocOffset, (int64)AlignedSize);
+
+		checkf((SIZE_T)OldOffset + AlignedSize <= SyncAllocReservedSize, TEXT("SyncObjectAllocator: Exceeded reserved size of %llu bytes"), SyncAllocReservedSize);
+
+		SyncAllocEnsureCommitted((SIZE_T)OldOffset + AlignedSize);
+
+		return SyncAllocBase + OldOffset;
+	}
+
+	INSTRUMENTATION_FUNCTION_ATTRIBUTES void SyncObjectAllocatorReset()
+	{
+		if (SyncAllocCommittedSize > 0)
+		{
+			VirtualFree(SyncAllocBase, (SIZE_T)SyncAllocCommittedSize, MEM_DECOMMIT);
+		}
+		SyncAllocOffset = 0;
+		SyncAllocCommittedSize = 0;
+	}
+
+	INSTRUMENTATION_FUNCTION_ATTRIBUTES void SyncObjectAllocatorAddRef()
+	{
+		FPlatformAtomics::InterlockedIncrement(&SyncAllocGlobalRefCount);
+	}
+
+	INSTRUMENTATION_FUNCTION_ATTRIBUTES void SyncObjectAllocatorRelease()
+	{
+		FPlatformAtomics::InterlockedDecrement(&SyncAllocGlobalRefCount);
+	}
+
+	INSTRUMENTATION_FUNCTION_ATTRIBUTES int32 SyncObjectAllocatorGetRefCount()
+	{
+		return SyncAllocGlobalRefCount;
+	}
+
+	INSTRUMENTATION_FUNCTION_ATTRIBUTES SIZE_T SyncObjectAllocatorGetCommittedSize()
+	{
+		return (SIZE_T)SyncAllocCommittedSize;
+	}
+
 	INSTRUMENTATION_FUNCTION_ATTRIBUTES void InitShadowMemory()
 	{
 		SYSTEM_INFO si;
