Support intrusive optional state for FSharedEventRef

#Synced-CL 52124292

diff --git a/Engine/Source/Runtime/Core/Public/HAL/Event.h b/Engine/Source/Runtime/Core/Public/HAL/Event.h
--- a/Engine/Source/Runtime/Core/Public/HAL/Event.h
+++ b/Engine/Source/Runtime/Core/Public/HAL/Event.h
@@ -4,6 +4,7 @@
 
 #include "CoreTypes.h"
 #include "Math/NumericLimits.h"
+#include "Misc/IntrusiveUnsetOptionalState.h"
 #include "Misc/Timespan.h"
 #include "Templates/Atomic.h"
 #include "Templates/SharedPointer.h"
@@ -171,6 +172,16 @@ public:
 		return Ptr.Get();
 	}
 
+	// Support intrusive TOptional state
+	constexpr static bool bHasIntrusiveUnsetOptionalState = true;
+	using IntrusiveUnsetOptionalStateType = FSharedEventRef;
+
+	explicit FSharedEventRef(FIntrusiveUnsetOptionalState) {}
+	bool operator==(FIntrusiveUnsetOptionalState) const
+	{
+		return !Ptr;
+	}
+
 private:
 	TSharedPtr<FEvent> Ptr;
 };
