Fix intrusive optional state for FSharedEventRef.  This prevents move operations from changing the optional state.

#Synced-CL 52124292

diff --git a/Engine/Source/Runtime/Core/Public/HAL/Event.h b/Engine/Source/Runtime/Core/Public/HAL/Event.h
--- a/Engine/Source/Runtime/Core/Public/HAL/Event.h
+++ b/Engine/Source/Runtime/Core/Public/HAL/Event.h
@@ -176,10 +176,12 @@ public:
 	constexpr static bool bHasIntrusiveUnsetOptionalState = true;
 	using IntrusiveUnsetOptionalStateType = FSharedEventRef;
 
-	explicit FSharedEventRef(FIntrusiveUnsetOptionalState) {}
-	bool operator==(FIntrusiveUnsetOptionalState) const
+	explicit FSharedEventRef(FIntrusiveUnsetOptionalState Tag) 
+		: Ptr(Tag)
+	{}
+	bool operator==(FIntrusiveUnsetOptionalState Tag) const
 	{
-		return !Ptr;
+		return Ptr == Tag;
 	}
 
 private:
