Added 8 bytes bin support to MB4

#Synced-CL 52054134

diff --git a/Engine/Source/Runtime/Core/Private/HAL/MallocBinned4.cpp b/Engine/Source/Runtime/Core/Private/HAL/MallocBinned4.cpp
--- a/Engine/Source/Runtime/Core/Private/HAL/MallocBinned4.cpp
+++ b/Engine/Source/Runtime/Core/Private/HAL/MallocBinned4.cpp
@@ -36,9 +36,10 @@ PRAGMA_DISABLE_UNSAFE_TYPECAST_WARNINGS
 
 // Bin sizes are based around getting the maximum amount of allocations per pool, with as little alignment waste as possible.
 // Bin sizes should be close to even divisors of the system page size, and well distributed.
-// They must be 16-byte aligned as well.
+// All bins, but 8, must be 16-byte aligned.
 constexpr uint16 MB4SmallBinSizesInternal[] =
 {
+	8,
 	16, 32, 48, 64, 80, 96, 112, 128, 144, 160, 176, 192, 208, // +16
 //	Bin		|	Divider | Slack leftover per page (64KB)
 	256,		// /256
@@ -217,7 +218,7 @@ void FMallocBinned4::FPoolInfo::ExhaustPoolIfNecessary()
 {
 	if (FirstFreeBlock->GetNumFreeBins() == 0)
 	{
-		FirstFreeBlock = FirstFreeBlock->NextFreeBlock;
+		FirstFreeBlock = FirstFreeBlock->GetNextFreeBlock();
 	}
 	check(!FirstFreeBlock || FirstFreeBlock->GetNumFreeBins() != 0);
 }
@@ -258,7 +259,7 @@ struct FMallocBinned4::Private
 			// Free a pooled allocation.
 			FFreeBlock* Free = (FFreeBlock*)Node;
 			Free->NumFreeBins = 1;
-			Free->NextFreeBlock = NodePool->FirstFreeBlock;
+			Free->SetNextFreeBlock(NodePool->FirstFreeBlock);
 			Free->BinSize = BinSize;
 			Free->CanaryAndForkState = Allocator.CurrentCanary;
 			Free->PoolIndex = InPoolIndex;
@@ -824,7 +825,7 @@ void FMallocBinned4::FPoolList::ValidateActivePools() const
 		const FPoolInfo* Pool = *PoolPtr;
 		check(Pool->PtrToPrevNext == PoolPtr);
 		check(Pool->FirstFreeBlock);
-		for (FFreeBlock* Free = Pool->FirstFreeBlock; Free; Free = Free->NextFreeBlock)
+		for (FFreeBlock* Free = Pool->FirstFreeBlock; Free; Free = Free->GetNextFreeBlock())
 		{
 			check(Free->GetNumFreeBins() > 0);
 		}
diff --git a/Engine/Source/Runtime/Core/Public/HAL/MallocBinned4.h b/Engine/Source/Runtime/Core/Public/HAL/MallocBinned4.h
--- a/Engine/Source/Runtime/Core/Public/HAL/MallocBinned4.h
+++ b/Engine/Source/Runtime/Core/Public/HAL/MallocBinned4.h
@@ -25,10 +25,10 @@
 
 #if AGGRESSIVE_MEMORY_SAVING
 #	define UE_MB4_MAX_SMALL_POOL_SIZE		(13104)		// Maximum bin size in SmallBinSizesInternal in cpp file
-#	define UE_MB4_SMALL_POOL_COUNT			48
+#	define UE_MB4_SMALL_POOL_COUNT			49
 #else
 #	define UE_MB4_MAX_SMALL_POOL_SIZE		(32768-16)	// Maximum bin size in SmallBinSizesInternal in cpp file
-#	define UE_MB4_SMALL_POOL_COUNT			51
+#	define UE_MB4_SMALL_POOL_COUNT			52
 #endif
 
 // If we are emulating forking on a windows server or are a linux server, enable support for avoiding dirtying pages owned by the parent. 
@@ -109,10 +109,10 @@ private:
 			: BinSize(InBinSize)
 			, PoolIndex(InPoolIndex)
 			, CanaryAndForkState(InCanary)
-			, NextFreeBlock(nullptr)
+			, NextFreeBlockIndex(InvalidNextFreeBlock)
 		{
-			check(InPoolIndex < MAX_uint8 && InBinSize <= MAX_uint16);
-			NumFreeBins = InPageSize / InBinSize;
+			check(PoolIndex == InPoolIndex);
+			NumFreeBins = uint16(InPageSize / InBinSize);
 			if (NumFreeBins * InBinSize + sizeof(FFreeBlock) > InPageSize)
 			{
 				NumFreeBins--;
@@ -128,26 +128,50 @@ private:
 		inline void* AllocateBin()
 		{
 			--NumFreeBins;
-			// this is a pointer to a whole 64KB block
+			// this is a pointer to a whole 64KB block - bins are allocated out from the end
 			if (IsAligned(this, UE_MB4_LARGE_ALLOC))
 			{
 				return (uint8*)this + UE_MB4_LARGE_ALLOC - (NumFreeBins + 1) * BinSize;
 			}
-
-			// this is a pointer to a single bin within the 64KB block
-			return (uint8*)this + (NumFreeBins * BinSize);
+			// this is a recycled single bin within the 64KB block
+			return (uint8*)this + NumFreeBins * BinSize;
 		}
 
-		uint16 BinSize;				// Size of the bins that this list points to
-		uint8 PoolIndex;			// Index of this pool
+		FORCEINLINE FFreeBlock* GetNextFreeBlock() const
+		{
+			if (NextFreeBlockIndex == InvalidNextFreeBlock)
+			{
+				return nullptr;
+			}
+			uint8* BlockBase = (uint8*)AlignDown(this, UE_MB4_LARGE_ALLOC);
+			return (FFreeBlock*)(BlockBase + NextFreeBlockIndex * UE_MBC_MIN_BIN_SIZE);
+		}
 
-		// Normally this value just functions as a canary to detect invalid memory state.
-		// When process forking is supported, it's still a canary but it has two valid values.
-		// One value is used pre-fork and one post-fork and the value is used to avoid freeing memory in pages shared with the parent process.
-		EBlockCanary CanaryAndForkState; 
+		FORCEINLINE void SetNextFreeBlock(FFreeBlock* Next)
+		{
+			if (Next)
+			{
+				uint8* BlockBase = (uint8*)AlignDown(this, UE_MB4_LARGE_ALLOC);
+				NextFreeBlockIndex = uint16((uint8*)Next - BlockBase) / UE_MBC_MIN_BIN_SIZE;
+			}
+			else
+			{
+				NextFreeBlockIndex = InvalidNextFreeBlock;
+			}
+		}
 
-		uint32 NumFreeBins;          // Number of consecutive free bins here, at least 1.
-		FFreeBlock*  NextFreeBlock;  // Next free block in another pool
+		//constexpr static uint32 MaxBitsForBinSize	 = 32 - FMath::CountLeadingZeros(UE_MB4_MAX_SMALL_POOL_SIZE));
+		constexpr static uint32 MaxBitsForPoolIndex = 32 - FMath::CountLeadingZeros(UE_MB4_SMALL_POOL_COUNT);
+		constexpr static uint32 MaxBitsForFreeBins = 32 - FMath::CountLeadingZeros(UE_MB4_LARGE_ALLOC / UE_MBC_MIN_BIN_SIZE);
+		constexpr static uint32 InvalidNextFreeBlock = 1u << MaxBitsForFreeBins;
+		static_assert(MaxBitsForPoolIndex <= 8, "PoolIndex should fit in 8 bits");
+		static_assert(MaxBitsForFreeBins + 1 <= 16, "NextFreeBlockIndex should fit in 16 bits");
+
+		uint16 BinSize;						// Size of the bins
+		uint8 PoolIndex;					// Index of this pool
+		EBlockCanary CanaryAndForkState;	// Canary value; two valid values when fork support is active
+		uint16 NumFreeBins;					// Number of consecutive free bins here, at least 1
+		uint16 NextFreeBlockIndex;			// Bin index within 64KB block of next recycled FFreeBlock, or InvalidNextFreeBlock
 	};
 
 	struct FPoolList
