Add GetAudioRenderThreadMask affinity io override

#rnx

1) Renamed all `FPlatformAffinity` to `FPlatformAffinityBase` to define a custom `FPlatformAffinity` that inherits `FPlatformAffinityBase` and adds platform-agnostic methods
2) Added `SetAudioRenderThreadMask` method to `FPlatformAffinity`
3) `GetAudioRenderThreadMask` will return a default platform affinity io or the overridden value otherwise
4) Added a private `AudioRenderThreadMask` that holds the overridden  value

#Synced-CL 52130646

diff --git a/Engine/Source/Runtime/Core/Private/HAL/PlatformAffinity.cpp b/Engine/Source/Runtime/Core/Private/HAL/PlatformAffinity.cpp
new file mode 100644
--- /dev/null
+++ b/Engine/Source/Runtime/Core/Private/HAL/PlatformAffinity.cpp
@@ -0,0 +1,21 @@
+// Copyright Epic Games, Inc. All Rights Reserved.
+
+#include "HAL/PlatformAffinity.h"
+
+#include "Misc/Optional.h"
+
+namespace UE::HAL
+{
+	/** Stores the overridden value for the audio render thread. */
+	static TOptional<uint64> AudioRenderThreadMask;
+}
+
+void FPlatformAffinity::SetAudioRenderThreadMask(const uint64 InNewValue)
+{
+	UE::HAL::AudioRenderThreadMask.Emplace(InNewValue);
+}
+
+const uint64 FPlatformAffinity::GetAudioRenderThreadMask()
+{
+	return UE::HAL::AudioRenderThreadMask.Get(FPlatformAffinityBase::GetAudioRenderThreadMask());
+}
diff --git a/Engine/Source/Runtime/Core/Public/Android/AndroidPlatformAffinity.h b/Engine/Source/Runtime/Core/Public/Android/AndroidPlatformAffinity.h
--- a/Engine/Source/Runtime/Core/Public/Android/AndroidPlatformAffinity.h
+++ b/Engine/Source/Runtime/Core/Public/Android/AndroidPlatformAffinity.h
@@ -84,4 +84,4 @@ public:
 	static uint64 RenderingThreadMask;
 };
 
-typedef FAndroidAffinity FPlatformAffinity;
+typedef FAndroidAffinity FPlatformAffinityBase;
diff --git a/Engine/Source/Runtime/Core/Public/HAL/PlatformAffinity.h b/Engine/Source/Runtime/Core/Public/HAL/PlatformAffinity.h
--- a/Engine/Source/Runtime/Core/Public/HAL/PlatformAffinity.h
+++ b/Engine/Source/Runtime/Core/Public/HAL/PlatformAffinity.h
@@ -6,8 +6,23 @@
 
 #include COMPILED_PLATFORM_HEADER(PlatformAffinity.h) // IWYU pragma: export
 
-struct FThreadAffinity 
-{ 
-	uint64 ThreadAffinityMask = FPlatformAffinity::GetNoAffinityMask(); 
-	uint16 ProcessorGroup = 0; 
+class FPlatformAffinity : public FPlatformAffinityBase
+{
+public:
+	/*
+	* Sets the new affinity mask used for the audio render threads
+	* @note It only takes effect for new audio render threads
+	*/
+	static CORE_API void SetAudioRenderThreadMask(const uint64 InNewValue);
+
+	/*
+	* Returns the affinity mask used for the audio render threads
+	*/
+	static const CORE_API uint64 GetAudioRenderThreadMask();
+};
+
+struct FThreadAffinity
+{
+	uint64 ThreadAffinityMask = FPlatformAffinity::GetNoAffinityMask();
+	uint16 ProcessorGroup = 0;
 };
diff --git a/Engine/Source/Runtime/Core/Public/IOS/IOSPlatformAffinity.h b/Engine/Source/Runtime/Core/Public/IOS/IOSPlatformAffinity.h
--- a/Engine/Source/Runtime/Core/Public/IOS/IOSPlatformAffinity.h
+++ b/Engine/Source/Runtime/Core/Public/IOS/IOSPlatformAffinity.h
@@ -15,4 +15,4 @@ public:
 	static const uint64 GetNoAffinityMask();
 };
 
-typedef FIOSPlatformAffinity FPlatformAffinity;
+typedef FIOSPlatformAffinity FPlatformAffinityBase;
diff --git a/Engine/Source/Runtime/Core/Public/Linux/LinuxPlatformAffinity.h b/Engine/Source/Runtime/Core/Public/Linux/LinuxPlatformAffinity.h
--- a/Engine/Source/Runtime/Core/Public/Linux/LinuxPlatformAffinity.h
+++ b/Engine/Source/Runtime/Core/Public/Linux/LinuxPlatformAffinity.h
@@ -3,4 +3,5 @@
 #pragma once
 
 #include "GenericPlatform/GenericPlatformAffinity.h" // IWYU pragma: export
-typedef FGenericPlatformAffinity FPlatformAffinity;
+
+typedef FGenericPlatformAffinity FPlatformAffinityBase;
diff --git a/Engine/Source/Runtime/Core/Public/Mac/MacPlatformAffinity.h b/Engine/Source/Runtime/Core/Public/Mac/MacPlatformAffinity.h
--- a/Engine/Source/Runtime/Core/Public/Mac/MacPlatformAffinity.h
+++ b/Engine/Source/Runtime/Core/Public/Mac/MacPlatformAffinity.h
@@ -4,4 +4,4 @@
 
 #include "Apple/ApplePlatformAffinity.h"
 
-typedef FApplePlatformAffinity FPlatformAffinity;
+typedef FApplePlatformAffinity FPlatformAffinityBase;
diff --git a/Engine/Source/Runtime/Core/Public/Windows/WindowsPlatformAffinity.h b/Engine/Source/Runtime/Core/Public/Windows/WindowsPlatformAffinity.h
--- a/Engine/Source/Runtime/Core/Public/Windows/WindowsPlatformAffinity.h
+++ b/Engine/Source/Runtime/Core/Public/Windows/WindowsPlatformAffinity.h
@@ -23,4 +23,4 @@ public:
 	}
 };
 
-typedef FWindowsPlatformAffinity FPlatformAffinity;
+typedef FWindowsPlatformAffinity FPlatformAffinityBase;
