Added runtime network status monitoring to platforms missing it (iOS, Android, Windows, Mac, Linux)

#Synced-CL 52121574

diff --git a/Engine/Source/Runtime/ApplicationCore/Private/Linux/LinuxApplication.cpp b/Engine/Source/Runtime/ApplicationCore/Private/Linux/LinuxApplication.cpp
--- a/Engine/Source/Runtime/ApplicationCore/Private/Linux/LinuxApplication.cpp
+++ b/Engine/Source/Runtime/ApplicationCore/Private/Linux/LinuxApplication.cpp
@@ -222,6 +222,9 @@ TSharedPtr< FLinuxWindow > FLinuxApplication::FindWindowBySDLWindow(SDL_Window *
 void FLinuxApplication::PumpMessages( const float TimeDelta )
 {
 	FPlatformApplicationMisc::PumpMessages( true );
+
+	// Check for netlink network change events and update status if needed.
+	FPlatformMisc::CheckNetworkConnectionEvents();
 }
 
 bool FLinuxApplication::IsCursorDirectlyOverSlateWindow() const
diff --git a/Engine/Source/Runtime/Core/Private/Android/AndroidPlatformMisc.cpp b/Engine/Source/Runtime/Core/Private/Android/AndroidPlatformMisc.cpp
--- a/Engine/Source/Runtime/Core/Private/Android/AndroidPlatformMisc.cpp
+++ b/Engine/Source/Runtime/Core/Private/Android/AndroidPlatformMisc.cpp
@@ -603,6 +603,9 @@ void FAndroidMisc::PlatformInit()
 
 	extern void AndroidThunkJava_AddNetworkListener();
 	AndroidThunkJava_AddNetworkListener();
+
+	// Set initial network connection status
+	FPlatformMisc::SetNetworkConnectionStatus(FPlatformMisc::NetworkConnectionTypeToStatus(GetNetworkConnectionType()));
 #endif
 
 	InitCpuThermalSensor();
@@ -3124,7 +3127,9 @@ void JNICALL UE::Jni::FGameActivity::nativeNetworkChanged(JNIEnv* jenv, jobject
 	{
 		FFunctionGraphTask::CreateAndDispatchWhenReady([]()
 		{
-			FCoreDelegates::OnNetworkConnectionChanged.Broadcast(FAndroidMisc::GetNetworkConnectionType());
+			const ENetworkConnectionType ConnectionType = FAndroidMisc::GetNetworkConnectionType();
+			FCoreDelegates::OnNetworkConnectionChanged.Broadcast(ConnectionType);
+			FPlatformMisc::SetNetworkConnectionStatus(FPlatformMisc::NetworkConnectionTypeToStatus(ConnectionType));
 		}, TStatId(), NULL, ENamedThreads::GameThread);
 	}
 }
diff --git a/Engine/Source/Runtime/Core/Private/Apple/ApplePlatformMisc.cpp b/Engine/Source/Runtime/Core/Private/Apple/ApplePlatformMisc.cpp
--- a/Engine/Source/Runtime/Core/Private/Apple/ApplePlatformMisc.cpp
+++ b/Engine/Source/Runtime/Core/Private/Apple/ApplePlatformMisc.cpp
@@ -560,17 +560,20 @@ void FApplePlatformMisc::PlatformInit()
 			}
 			UE::FApplePlatformMisc::CachedNetworkPath = InPath;
 			nw_retain(UE::FApplePlatformMisc::CachedNetworkPath);
+
+			ENetworkConnectionType ConnectionType = UE::FApplePlatformMisc::CheckNetworkConnectionType();
+
 			if (FTaskGraphInterface::IsRunning())
 			{
-				ENetworkConnectionType ConnectionType = UE::FApplePlatformMisc::CheckNetworkConnectionType();
 				if (ConnectionType != UE::FApplePlatformMisc::LastReportedNetworkType)
 				{
 					UE::FApplePlatformMisc::LastReportedNetworkType = ConnectionType;
 					FFunctionGraphTask::CreateAndDispatchWhenReady([ConnectionType]() {
 						FCoreDelegates::OnNetworkConnectionChanged.Broadcast(ConnectionType);
+						FPlatformMisc::SetNetworkConnectionStatus(FPlatformMisc::NetworkConnectionTypeToStatus(ConnectionType));
 					}, TStatId(), nullptr, ENamedThreads::GameThread);
 				}
-				
+
 				// Adding a reference, so path survives to the task body.
 				nw_retain(InPath);
 				FFunctionGraphTask::CreateAndDispatchWhenReady(
@@ -580,10 +583,18 @@ void FApplePlatformMisc::PlatformInit()
 					UE::FApplePlatformMisc::NotifyAboutNetworkConnectionCharacteristicsChange(InPath);
 				}, TStatId());
 			}
+			else
+			{
+				// Task graph not running yet (early startup). Set the initial status synchronously so it isn't silently dropped.
+				UE::FApplePlatformMisc::LastReportedNetworkType = ConnectionType;
+				FPlatformMisc::SetNetworkConnectionStatus(FPlatformMisc::NetworkConnectionTypeToStatus(ConnectionType));
+			}
 		}
 	);
 	
 	nw_path_monitor_set_queue(UE::FApplePlatformMisc::PathMonitor, dispatch_get_main_queue());
 	nw_path_monitor_start(UE::FApplePlatformMisc::PathMonitor);
 
+	// The nw_path_monitor callback will fire shortly after start.
+	// If the task graph is not yet running at that point, the callback sets the initial network status synchronously
 }
diff --git a/Engine/Source/Runtime/Core/Private/Linux/LinuxPlatformMisc.cpp b/Engine/Source/Runtime/Core/Private/Linux/LinuxPlatformMisc.cpp
--- a/Engine/Source/Runtime/Core/Private/Linux/LinuxPlatformMisc.cpp
+++ b/Engine/Source/Runtime/Core/Private/Linux/LinuxPlatformMisc.cpp
@@ -4,8 +4,143 @@
 	LinuxPlatformMisc.cpp: Linux implementations of misc platform functions
 =============================================================================*/
 
-/* Add back if we overload any functions for GenericPlatform here
 #include "Linux/LinuxPlatformMisc.h"
-#include "Containers/Array.h"
-#include "Containers/UnrealString.h"
-*/
+#include "Logging/LogMacros.h"
+#include "Misc/CoreDelegates.h"
+
+#include <ifaddrs.h>
+#include <net/if.h>
+#include <sys/socket.h>
+#include <linux/netlink.h>
+#include <linux/rtnetlink.h>
+#include <unistd.h>
+
+/** Netlink socket file descriptor for receiving network change events. -1 when not initialized. */
+static int GNetlinkSocket = -1;
+
+void FLinuxPlatformMisc::PlatformInit()
+{
+	FUnixPlatformMisc::PlatformInit();
+	StartNetworkConnectionMonitoring();
+}
+
+void FLinuxPlatformMisc::PlatformTearDown()
+{
+	StopNetworkConnectionMonitoring();
+	FUnixPlatformMisc::PlatformTearDown();
+}
+
+ENetworkConnectionType FLinuxPlatformMisc::GetNetworkConnectionType()
+{
+	struct ifaddrs* InterfaceAddrs = nullptr;
+	if (getifaddrs(&InterfaceAddrs) != 0)
+	{
+		return ENetworkConnectionType::Unknown;
+	}
+
+	ENetworkConnectionType Result = ENetworkConnectionType::None;
+
+	for (struct ifaddrs* CurrentAddr = InterfaceAddrs; CurrentAddr != nullptr; CurrentAddr = CurrentAddr->ifa_next)
+	{
+		// Skip interfaces that are not up and running
+		if ((CurrentAddr->ifa_flags & (IFF_UP | IFF_RUNNING)) != (IFF_UP | IFF_RUNNING))
+		{
+			continue;
+		}
+
+		// Skip loopback interfaces
+		if (CurrentAddr->ifa_flags & IFF_LOOPBACK)
+		{
+			continue;
+		}
+
+		// Must have a valid address (AF_INET or AF_INET6)
+		if (CurrentAddr->ifa_addr == nullptr)
+		{
+			continue;
+		}
+
+		int Family = CurrentAddr->ifa_addr->sa_family;
+		if (Family != AF_INET && Family != AF_INET6)
+		{
+			continue;
+		}
+
+		// Found a non-loopback interface that is up and running with a valid address
+		Result = ENetworkConnectionType::Ethernet;
+		break;
+	}
+
+	freeifaddrs(InterfaceAddrs);
+	return Result;
+}
+
+void FLinuxPlatformMisc::StartNetworkConnectionMonitoring()
+{
+	check(GNetlinkSocket == -1);
+
+	GNetlinkSocket = socket(AF_NETLINK, SOCK_DGRAM | SOCK_NONBLOCK | SOCK_CLOEXEC, NETLINK_ROUTE);
+	if (GNetlinkSocket == -1)
+	{
+		UE_LOG(LogInit, Warning, TEXT("Failed to create netlink socket for network monitoring (errno=%d)"), errno);
+		return;
+	}
+
+	struct sockaddr_nl Addr;
+	FMemory::Memzero(Addr);
+	Addr.nl_family = AF_NETLINK;
+	Addr.nl_groups = RTMGRP_LINK | RTMGRP_IPV4_IFADDR | RTMGRP_IPV6_IFADDR;
+
+	if (bind(GNetlinkSocket, (struct sockaddr*)&Addr, sizeof(Addr)) < 0)
+	{
+		UE_LOG(LogInit, Warning, TEXT("Failed to bind netlink socket for network monitoring (errno=%d)"), errno);
+		close(GNetlinkSocket);
+		GNetlinkSocket = -1;
+		return;
+	}
+
+	// Set the initial network connection status.
+	SetNetworkConnectionStatus(NetworkConnectionTypeToStatus(GetNetworkConnectionType()));
+}
+
+void FLinuxPlatformMisc::StopNetworkConnectionMonitoring()
+{
+	if (GNetlinkSocket != -1)
+	{
+		close(GNetlinkSocket);
+		GNetlinkSocket = -1;
+	}
+}
+
+void FLinuxPlatformMisc::CheckNetworkConnectionEvents()
+{
+	if (GNetlinkSocket == -1)
+	{
+		return;
+	}
+
+	// Non-blocking read: drain any pending netlink messages.
+	// We only care that something changed, not the specific message contents.
+	char Buffer[4096];
+	bool bHasEvents = false;
+
+	while (true)
+	{
+		ssize_t BytesRead = recv(GNetlinkSocket, Buffer, sizeof(Buffer), MSG_DONTWAIT);
+
+		if (BytesRead <= 0)
+		{
+			break;
+		}
+
+		bHasEvents = true;
+	}
+
+	if (bHasEvents)
+	{
+		const ENetworkConnectionType ConnectionType = GetNetworkConnectionType();
+		
+		FCoreDelegates::OnNetworkConnectionChanged.Broadcast(ConnectionType);
+		SetNetworkConnectionStatus(NetworkConnectionTypeToStatus(ConnectionType));
+	}
+}
diff --git a/Engine/Source/Runtime/Core/Private/Windows/WindowsPlatformMisc.cpp b/Engine/Source/Runtime/Core/Private/Windows/WindowsPlatformMisc.cpp
--- a/Engine/Source/Runtime/Core/Private/Windows/WindowsPlatformMisc.cpp
+++ b/Engine/Source/Runtime/Core/Private/Windows/WindowsPlatformMisc.cpp
@@ -44,6 +44,7 @@
 #include "GenericPlatform/GenericPlatformDriver.h"
 #include "HAL/ThreadHeartBeat.h"
 #include "ProfilingDebugging/ExternalProfiler.h"
+#include "Async/TaskGraphInterfaces.h"
 
 // Resource includes.
 #include "Runtime/Launch/Resources/Windows/resource.h"
@@ -60,6 +61,7 @@ THIRD_PARTY_INCLUDES_START
 	#include <shlwapi.h>
 	#include <IPHlpApi.h>
 	#include <VersionHelpers.h>
+	#include <netlistmgr.h>
 #include "Windows/AllowWindowsPlatformAtomics.h"
 	#include <comdef.h>
 	#include <Wbemidl.h>
@@ -1091,10 +1093,14 @@ void FWindowsPlatformMisc::PlatformInit()
 
 	// Collecting storage information using dcom services.
 	CollectStorageInformation();
+
+	// Start listening for network connectivity changes via NLM events.
+	StartNetworkConnectionMonitoring();
 }
 
 void FWindowsPlatformMisc::PlatformTearDown()
 {
+	StopNetworkConnectionMonitoring();
 	FPlatformProcess::CeaseBeingFirstInstance();
 }
 
@@ -4547,4 +4553,285 @@ const class TMap<FString, FString>& FWindowsPlatformMisc::GetConfigRuleVars()
 #endif // WITH_APPLICATION_CORE && WITH_WINDOWS_CONFIGRULES
 }
 
+/**
+ * COM event listener that receives network connectivity change notifications from NLM.
+ * Registered in StartNetworkConnectionMonitoring, unregistered in StopNetworkConnectionMonitoring.
+ */
+class FNLMEventListener : public INetworkListManagerEvents
+{
+public:
+	FNLMEventListener() : RefCount(1), Cookie(0), ConnectionPoint(nullptr) {}
+	virtual ~FNLMEventListener() = default;
+
+	// IUnknown interface
+
+	HRESULT STDMETHODCALLTYPE QueryInterface(REFIID riid, void** ppvObject) override
+	{
+		if (riid == IID_IUnknown || riid == IID_INetworkListManagerEvents)
+		{
+			*ppvObject = this;
+			AddRef();
+			return S_OK;
+		}
+		*ppvObject = nullptr;
+		return E_NOINTERFACE;
+	}
+
+	ULONG STDMETHODCALLTYPE AddRef() override
+	{
+		// We use the PlatformAtomics methods to avoid having to add this to AllowWindowsPlatformTypes
+		return FPlatformAtomics::InterlockedIncrement(&RefCount);
+	}
+
+	ULONG STDMETHODCALLTYPE Release() override
+	{
+		// We use the PlatformAtomics methods to avoid having to add this to AllowWindowsPlatformTypes
+		int32 Ref = FPlatformAtomics::InterlockedDecrement(&RefCount);
+		if (Ref == 0)
+		{
+			delete this;
+		}
+		return Ref;
+	}
+
+	// ~ IUnknown interface
+
+	// INetworkListManagerEvents
+
+	HRESULT STDMETHODCALLTYPE ConnectivityChanged(NLM_CONNECTIVITY NewConnectivity) override
+	{
+		if (FTaskGraphInterface::IsRunning())
+		{
+			FFunctionGraphTask::CreateAndDispatchWhenReady([]()
+			{
+				const ENetworkConnectionType ConnectionType = FWindowsPlatformMisc::GetNetworkConnectionType();
+				
+				FCoreDelegates::OnNetworkConnectionChanged.Broadcast(ConnectionType);
+				FPlatformMisc::SetNetworkConnectionStatus(FPlatformMisc::NetworkConnectionTypeToStatus(ConnectionType));
+			}, TStatId(), nullptr, ENamedThreads::GameThread);
+		}
+		return S_OK;
+	}
+
+	// ~ INetworkListManagerEvents
+
+	/** Subscribe to NLM connectivity events. Returns true on success. */
+	bool Advise(INetworkListManager* NetworkListManager)
+	{
+		IConnectionPointContainer* Container = nullptr;
+		HRESULT Result = NetworkListManager->QueryInterface(IID_IConnectionPointContainer, (void**)&Container);
+		if (FAILED(Result) || Container == nullptr)
+		{
+			return false;
+		}
+
+		Result = Container->FindConnectionPoint(IID_INetworkListManagerEvents, &ConnectionPoint);
+		Container->Release();
+		if (FAILED(Result) || ConnectionPoint == nullptr)
+		{
+			return false;
+		}
+
+		Result = ConnectionPoint->Advise(this, &Cookie);
+		if (FAILED(Result))
+		{
+			ConnectionPoint->Release();
+			ConnectionPoint = nullptr;
+			return false;
+		}
+
+		return true;
+	}
+
+	/** Unsubscribe from NLM connectivity events. */
+	void Unadvise()
+	{
+		if (ConnectionPoint != nullptr)
+		{
+			if (Cookie != 0)
+			{
+				ConnectionPoint->Unadvise(Cookie);
+				Cookie = 0;
+			}
+			ConnectionPoint->Release();
+			ConnectionPoint = nullptr;
+		}
+	}
+
+private:
+	// Used by the IUnknown interface, type required by InterlockedIncrement/Decrement mehtods
+	volatile int32 RefCount;
+
+	// Token returned by Advise which identifies our registration
+	DWORD Cookie;
+
+	// Intermediary between the event source (INetworkListManager) and this listener
+	IConnectionPoint* ConnectionPoint;
+};
+
+static FNLMEventListener* GNLMEventListener = nullptr;
+static INetworkListManager* GNetworkListManager = nullptr;
+
+/**
+ * Queries NLM network connections to determine if the active connection is WiFi or wired Ethernet
+ */
+static ENetworkConnectionType DetermineConnectionTypeFromNLM(INetworkListManager* NetworkListManager)
+{
+	// Get an enumerator over all NLM network connections
+	IEnumNetworkConnections* EnumConnections = nullptr;
+	HRESULT Result = NetworkListManager->GetNetworkConnections(&EnumConnections);
+	if (FAILED(Result) || EnumConnections == nullptr)
+	{
+		return ENetworkConnectionType::Ethernet;
+	}
+
+	bool bHasWiFi = false;
+	bool bHasEthernet = false;
+
+	// Iterate each connection and check only those that are currently connected
+	INetworkConnection* Connection = nullptr;
+	while (EnumConnections->Next(1, &Connection, nullptr) == S_OK)
+	{
+		VARIANT_BOOL IsConnected = VARIANT_FALSE;
+		if (SUCCEEDED(Connection->get_IsConnected(&IsConnected)) && IsConnected != VARIANT_FALSE)
+		{
+			// Get the adapter GUID that backs this NLM connection
+			GUID AdapterId;
+			if (SUCCEEDED(Connection->GetAdapterId(&AdapterId)))
+			{
+				// Convert adapter GUID -> LUID -> interface index.
+				// We need the index because MIB_IFROW (used by GetIfEntry) is keyed by dwIndex
+				
+				NET_LUID Luid;
+				NET_IFINDEX IfIndex;
+
+				if (ConvertInterfaceGuidToLuid(&AdapterId, &Luid) == NO_ERROR &&
+					ConvertInterfaceLuidToIndex(&Luid, &IfIndex) == NO_ERROR)
+				{
+					// Query the interface row to read the physical adapter type
+					
+					MIB_IFROW IfRow = {};
+					IfRow.dwIndex = IfIndex;
+
+					if (GetIfEntry(&IfRow) == NO_ERROR)
+					{
+						if (IfRow.dwType == IF_TYPE_IEEE80211)
+						{
+							bHasWiFi = true;
+						}
+						else if (IfRow.dwType == IF_TYPE_ETHERNET_CSMACD)
+						{
+							bHasEthernet = true;
+						}
+					}
+				}
+			}
+		}
+
+		Connection->Release();
+	}
+
+	EnumConnections->Release();
+
+	// When both WiFi and Ethernet are active (e.g. docked laptop), prefer Ethernet, as the wired connection is typically the primary route
+	
+	if (bHasEthernet)
+	{
+		return ENetworkConnectionType::Ethernet;
+	}
+
+	if (bHasWiFi)
+	{
+		return ENetworkConnectionType::WiFi;
+	}
+
+	// Fallback: if no recognized adapter type was found, default to Ethernet
+	return ENetworkConnectionType::Ethernet;
+}
+
+ENetworkConnectionType FWindowsPlatformMisc::GetNetworkConnectionType()
+{
+	ENetworkConnectionType Result = ENetworkConnectionType::Unknown;
+
+	// Use the cached NLM instance if available, otherwise create a temporary one.
+	INetworkListManager* NetworkListManager = GNetworkListManager;
+	bool bCreatedTemporaryNLM = false;
+
+	if (NetworkListManager == nullptr)
+	{
+		HRESULT CoCreateInstanceResult = CoCreateInstance(CLSID_NetworkListManager, nullptr, CLSCTX_ALL, IID_INetworkListManager, (void**)&NetworkListManager);
+		if (FAILED(CoCreateInstanceResult) || NetworkListManager == nullptr)
+		{
+			return Result;
+		}
+		bCreatedTemporaryNLM = true;
+	}
+
+	NLM_CONNECTIVITY Connectivity;
+	HRESULT GetConnectivityResult = NetworkListManager->GetConnectivity(&Connectivity);
+	if (SUCCEEDED(GetConnectivityResult))
+	{
+		if (Connectivity & (NLM_CONNECTIVITY_IPV4_INTERNET | NLM_CONNECTIVITY_IPV6_INTERNET |
+			NLM_CONNECTIVITY_IPV4_LOCALNETWORK | NLM_CONNECTIVITY_IPV6_LOCALNETWORK))
+		{
+			// Enumerate NLM connections and resolve adapter types to distinguish WiFi from Ethernet
+			Result = DetermineConnectionTypeFromNLM(NetworkListManager);
+		}
+		else
+		{
+			Result = ENetworkConnectionType::None;
+		}
+	}
+
+	if (bCreatedTemporaryNLM)
+	{
+		NetworkListManager->Release();
+	}
+
+	return Result;
+}
+
+void FWindowsPlatformMisc::StartNetworkConnectionMonitoring()
+{
+	check(GNetworkListManager == nullptr && GNLMEventListener == nullptr);
+
+	HRESULT CoCreateInstanceResult = CoCreateInstance(CLSID_NetworkListManager, nullptr, CLSCTX_ALL, IID_INetworkListManager, (void**)&GNetworkListManager);
+	if (FAILED(CoCreateInstanceResult) || GNetworkListManager == nullptr)
+	{
+		UE_LOG(LogInit, Warning, TEXT("Failed to create INetworkListManager for network monitoring (HRESULT=0x%08X)"), CoCreateInstanceResult);
+		GNetworkListManager = nullptr;
+		return;
+	}
+
+	GNLMEventListener = new FNLMEventListener();
+	if (!GNLMEventListener->Advise(GNetworkListManager))
+	{
+		UE_LOG(LogInit, Warning, TEXT("Failed to subscribe to INetworkListManagerEvents"));
+		GNLMEventListener->Release();
+		GNLMEventListener = nullptr;
+		GNetworkListManager->Release();
+		GNetworkListManager = nullptr;
+		return;
+	}
+
+	// Set the initial network connection status.
+	SetNetworkConnectionStatus(NetworkConnectionTypeToStatus(GetNetworkConnectionType()));
+}
+
+void FWindowsPlatformMisc::StopNetworkConnectionMonitoring()
+{
+	if (GNLMEventListener != nullptr)
+	{
+		GNLMEventListener->Unadvise();
+		GNLMEventListener->Release();
+		GNLMEventListener = nullptr;
+	}
+
+	if (GNetworkListManager != nullptr)
+	{
+		GNetworkListManager->Release();
+		GNetworkListManager = nullptr;
+	}
+}
+
 #undef UE_CPUID
diff --git a/Engine/Source/Runtime/Core/Public/GenericPlatform/GenericPlatformMisc.h b/Engine/Source/Runtime/Core/Public/GenericPlatform/GenericPlatformMisc.h
--- a/Engine/Source/Runtime/Core/Public/GenericPlatform/GenericPlatformMisc.h
+++ b/Engine/Source/Runtime/Core/Public/GenericPlatform/GenericPlatformMisc.h
@@ -1686,6 +1686,22 @@ public:
 	 */
 	static CORE_API void SetNetworkConnectionStatus(ENetworkConnectionStatus NewNetworkConnectionStatus);
 
+	/**
+	 * Converts a network connection type to a network connection status.
+	 * Maps None/AirplaneMode to Disabled, all known connection types to Connected.
+	 */
+	static ENetworkConnectionStatus NetworkConnectionTypeToStatus(ENetworkConnectionType ConnectionType)
+	{
+		switch (ConnectionType)
+		{
+		case ENetworkConnectionType::None:
+		case ENetworkConnectionType::AirplaneMode:
+			return ENetworkConnectionStatus::Disabled;
+		default:
+			return ENetworkConnectionStatus::Connected;
+		}
+	}
+
 	/**
 	 * Returns whether WiFi connection is currently active
 	 */
diff --git a/Engine/Source/Runtime/Core/Public/Linux/LinuxPlatformMisc.h b/Engine/Source/Runtime/Core/Public/Linux/LinuxPlatformMisc.h
--- a/Engine/Source/Runtime/Core/Public/Linux/LinuxPlatformMisc.h
+++ b/Engine/Source/Runtime/Core/Public/Linux/LinuxPlatformMisc.h
@@ -12,9 +12,28 @@
  * Linux implementation of the misc OS functions
  */
 
-// A stub for now. Will allow to override any functions if needed
 struct FLinuxPlatformMisc : public FUnixPlatformMisc
 {
+	static CORE_API void PlatformInit();
+	static CORE_API void PlatformTearDown();
+
+	/**
+	 * Returns the current network connection type by checking non-loopback interfaces via getifaddrs().
+	 */
+	static CORE_API ENetworkConnectionType GetNetworkConnectionType();
+
+	/**
+	 * Checks the netlink socket for pending network change events and updates
+	 * connection status if a change occurred. Should be called from PumpMessages().
+	 */
+	static CORE_API void CheckNetworkConnectionEvents();
+
+private:
+	/** Creates a NETLINK_ROUTE socket to listen for interface and address changes. */
+	static void StartNetworkConnectionMonitoring();
+
+	/** Closes the netlink socket. */
+	static void StopNetworkConnectionMonitoring();
 };
 
 typedef FLinuxPlatformMisc FPlatformMisc;
diff --git a/Engine/Source/Runtime/Core/Public/Windows/WindowsPlatformMisc.h b/Engine/Source/Runtime/Core/Public/Windows/WindowsPlatformMisc.h
--- a/Engine/Source/Runtime/Core/Public/Windows/WindowsPlatformMisc.h
+++ b/Engine/Source/Runtime/Core/Public/Windows/WindowsPlatformMisc.h
@@ -325,6 +325,23 @@ struct FWindowsPlatformMisc
 		return true;
 	}
 
+	/**
+	 * Returns the current network connection type using NLM (Network List Manager).
+	 */
+	static CORE_API ENetworkConnectionType GetNetworkConnectionType();
+
+	/**
+	 * Starts event-driven network connection monitoring via INetworkListManagerEvents.
+	 * Called from PlatformInit(). Updates network connection status on connectivity changes.
+	 */
+	static CORE_API void StartNetworkConnectionMonitoring();
+
+	/**
+	 * Stops network connection monitoring and releases NLM COM resources.
+	 * Called from PlatformTearDown().
+	 */
+	static CORE_API void StopNetworkConnectionMonitoring();
+
 	/**
 	 * Returns whether the device is considered a gaming handheld device. It's not a dynamic property, the status is cached at startup.
 	 */
