Adjust serialization code for text history with dev notes so that standalone programs like Multi-user server to handle data with dev notes in the serialization blobs.

#Synced-CL 52151907

diff --git a/Engine/Source/Runtime/Core/Private/Internationalization/TextHistory.cpp b/Engine/Source/Runtime/Core/Private/Internationalization/TextHistory.cpp
--- a/Engine/Source/Runtime/Core/Private/Internationalization/TextHistory.cpp
+++ b/Engine/Source/Runtime/Core/Private/Internationalization/TextHistory.cpp
@@ -912,12 +912,28 @@ void FTextHistory_Base::Serialize(FStructuredArchive::FRecord Record)
 		Record << SA_VALUE(TEXT("SourceString"), SourceString);
 	}
 
-#if WITH_EDITORONLY_DATA
+#if WITH_EDITORONLY_DATA || IS_PROGRAM
 	BaseArchive.UsingCustomVersion(FFortniteMainBranchObjectVersion::GUID);
 	if (BaseArchive.CustomVer(FFortniteMainBranchObjectVersion::GUID) >= FFortniteMainBranchObjectVersion::AddDevNotesToFText && !BaseArchive.IsFilterEditorOnly())
 	{
-		Record << SA_VALUE(TEXT("DevNotes"), DevNotes);
-	}
+		if (BaseArchive.IsLoading())
+		{
+			FString LoadedDevNotes;
+			Record << SA_VALUE(TEXT("DevNotes"), LoadedDevNotes);
+#if WITH_EDITORONLY_DATA
+			DevNotes = MoveTemp(LoadedDevNotes);
+#endif
+		}
+		else if (BaseArchive.IsSaving())
+		{
+#if WITH_EDITORONLY_DATA
+			Record << SA_VALUE(TEXT("DevNotes"), DevNotes);
+#else
+			FString EmptyDevNotes;
+			Record << SA_VALUE(TEXT("DevNotes"), EmptyDevNotes);
+#endif
+		}
+	}	
 #endif
 }
 
