Added missing case statements

#Synced-CL 52121574

diff --git a/Engine/Source/Runtime/Core/Public/GenericPlatform/GenericPlatformMisc.h b/Engine/Source/Runtime/Core/Public/GenericPlatform/GenericPlatformMisc.h
--- a/Engine/Source/Runtime/Core/Public/GenericPlatform/GenericPlatformMisc.h
+++ b/Engine/Source/Runtime/Core/Public/GenericPlatform/GenericPlatformMisc.h
@@ -1697,6 +1697,12 @@ public:
 		case ENetworkConnectionType::None:
 		case ENetworkConnectionType::AirplaneMode:
 			return ENetworkConnectionStatus::Disabled;
+		case ENetworkConnectionType::Cell:
+		case ENetworkConnectionType::WiFi:
+		case ENetworkConnectionType::WiMAX:
+		case ENetworkConnectionType::Bluetooth:
+		case ENetworkConnectionType::Ethernet:
+		case ENetworkConnectionType::Unknown:
 		default:
 			return ENetworkConnectionStatus::Connected;
 		}
