Added missing check to ensure COM initialization

#Synced-CL 52171797

diff --git a/Engine/Source/Runtime/Core/Private/Windows/WindowsPlatformMisc.cpp b/Engine/Source/Runtime/Core/Private/Windows/WindowsPlatformMisc.cpp
--- a/Engine/Source/Runtime/Core/Private/Windows/WindowsPlatformMisc.cpp
+++ b/Engine/Source/Runtime/Core/Private/Windows/WindowsPlatformMisc.cpp
@@ -4795,6 +4795,13 @@ void FWindowsPlatformMisc::StartNetworkConnectionMonitoring()
 {
 	check(GNetworkListManager == nullptr && GNLMEventListener == nullptr);
 
+	// COM must be initialized on this thread before we can create COM objects.
+	if (!FWindowsPlatformMisc::CoInitialize())
+	{
+		UE_LOG(LogInit, Warning, TEXT("Failed to initialize COM for network monitoring"));
+		return;
+	}
+
 	HRESULT CoCreateInstanceResult = CoCreateInstance(CLSID_NetworkListManager, nullptr, CLSCTX_ALL, IID_INetworkListManager, (void**)&GNetworkListManager);
 	if (FAILED(CoCreateInstanceResult) || GNetworkListManager == nullptr)
 	{
