Remove deprecated UE4 version globals as they are no longer in use
#Synced-CL 52101958

diff --git a/Engine/Source/Runtime/Core/Private/UObject/ObjectVersion.cpp b/Engine/Source/Runtime/Core/Private/UObject/ObjectVersion.cpp
--- a/Engine/Source/Runtime/Core/Private/UObject/ObjectVersion.cpp
+++ b/Engine/Source/Runtime/Core/Private/UObject/ObjectVersion.cpp
@@ -13,12 +13,10 @@
 #include "Serialization/CompactBinaryWriter.h"
 
 // @see ObjectVersion.h for the list of changes/defines
-const int32 GPackageFileLicenseeUE4Version = GPackageFileLicenseeUEVersion;
 
 const FPackageFileVersion GPackageFileUEVersion(VER_LATEST_ENGINE_UE4, EUnrealEngineObjectUE5Version::AUTOMATIC_VERSION);
 const FPackageFileVersion GOldestLoadablePackageFileUEVersion = FPackageFileVersion::CreateUE4Version(VER_UE4_OLDEST_LOADABLE_PACKAGE);
 
-const int32 GPackageFileUE4Version = VER_LATEST_ENGINE_UE4;
 const int32 GPackageFileLicenseeUEVersion = VER_LATEST_ENGINE_LICENSEEUE4;
 
 FPackageFileVersion FPackageFileVersion::CreateUE4Version(int32 Version)
diff --git a/Engine/Source/Runtime/Core/Public/UObject/ObjectVersion.h b/Engine/Source/Runtime/Core/Public/UObject/ObjectVersion.h
--- a/Engine/Source/Runtime/Core/Public/UObject/ObjectVersion.h
+++ b/Engine/Source/Runtime/Core/Public/UObject/ObjectVersion.h
@@ -882,13 +882,6 @@ public:
 #define VER_MIN_SCRIPTVM_LICENSEEUE4							(VER_LIC_NONE)
 
 // Version access.
-
-UE_DEPRECATED(5.0, "Use GPackageFileUEVersion instead. See the @FPackageFileVersion documentation for further details")
-extern const CORE_API int32					GPackageFileUE4Version;				// Unreal4 Version Number.
-
-UE_DEPRECATED(5.0, "Use GPackageFileLicenseeUEVersion instead")
-extern const CORE_API int32					GPackageFileLicenseeUE4Version;		// Licensee Version Number.
-
 extern const CORE_API FPackageFileVersion	GPackageFileUEVersion;				// Unreal Version Number.
 extern const CORE_API FPackageFileVersion	GOldestLoadablePackageFileUEVersion;// The oldest Unreal Version number that we can load.
 extern const CORE_API int32					GPackageFileLicenseeUEVersion;		// Licensee Version Number.
