Moved out block header from the 64kb MB2 block into a separate lookup table.
This allows to pack power of two allocations like 1024, 4096 into 64kb block without a waste and saves a significant amount of memory

#Synced-CL 52054134

diff --git a/Engine/Source/Runtime/Core/Private/HAL/MallocBinned4.cpp b/Engine/Source/Runtime/Core/Private/HAL/MallocBinned4.cpp
--- a/Engine/Source/Runtime/Core/Private/HAL/MallocBinned4.cpp
+++ b/Engine/Source/Runtime/Core/Private/HAL/MallocBinned4.cpp
@@ -53,19 +53,19 @@ constexpr uint16 MB4SmallBinSizesInternal[] =
 	720,		// /91	16b
 	816,		// /80	256b
 	912,		// /71	784b
-	1024-16,	// /64
+	1024,		// /64
 	1168,		// /56	128b
 	1392,		// /47	112b
 	1520,		// /43	176b
 	1680,		// /39	16b
 	1872,		// /35	16b
-	2048-16,	// /32
+	2048,		// /32
 	2256,		// /29	112b
 	2608,		// /25	336b
 	2976,		// /22	64b
 	3264,		// /20	256b
 	3632,		// /18	160b
-	4096-16,	// /16
+	4096,		// /16
 	4368,		// /15	16b
 	4672,		// /14	128b
 	5040,		// /13	16b
@@ -73,14 +73,14 @@ constexpr uint16 MB4SmallBinSizesInternal[] =
 	5952,		// /11	64b
 	6544,		// /10	96b
 	7280,		// /9	16b
-	8192-16,	// /8
+	8192,		// /8
 	9360,		// /7	16b
 	10912,		// /6	64b
 	13104,		// /5	16b
 #if !AGGRESSIVE_MEMORY_SAVING
-	16384-16,	// /4
+	16384,		// /4
 	21840,		// /3	16b
-	32768-16	// /2
+	32768		// /2
 #endif
 };
 
@@ -262,7 +262,6 @@ struct FMallocBinned4::Private
 			Free->SetNextFreeBlock(NodePool->FirstFreeBlock);
 			Free->BinSize = BinSize;
 			Free->CanaryAndForkState = Allocator.CurrentCanary;
-			Free->PoolIndex = InPoolIndex;
 			NodePool->FirstFreeBlock = Free;
 
 			UE_MBC_UPDATE_STATS(--Table.TotalUsedBins);
@@ -272,22 +271,18 @@ struct FMallocBinned4::Private
 			if (--NodePool->Taken == 0)
 			{
 				NodePool->SetCanary(FPoolInfo::ECanary::Unassigned, true, false);
-				FFreeBlock* BasePtrOfNode = GetPoolHeaderFromPointer(Node);
+				void* BasePtrOfNode = (void*)AlignDown(Node, UE_MB4_LARGE_ALLOC);
 
 				// Free the OS memory.
 				NodePool->Unlink();
 				{
 					UE::TUniqueLock InnerLock(Allocator.ExternalAllocMutex);
+					Allocator.UnregisterBlockPool(BasePtrOfNode);
 					Allocator.CachedOSPageAllocator.Free(BasePtrOfNode, GMB4PageSize);
 				}
 
 #if UE_MB4_ALLOCATOR_STATS
-				int NumBinsInBlock = GMB4PageSize / BinSize;
-				if (NumBinsInBlock * BinSize + sizeof(FFreeBlock) > GMB4PageSize)
-				{
-					NumBinsInBlock--;
-				}
-
+				const int NumBinsInBlock = GMB4PageSize / BinSize;
 				Table.TotalAllocatedBins -= NumBinsInBlock;
 				Table.TotalAllocatedMem -= GMB4PageSize;
 				AllocatedOSSmallPoolMemory.fetch_sub((int64)GMB4PageSize, std::memory_order_relaxed);
@@ -356,7 +351,7 @@ FMallocBinned4::FPoolInfo& FMallocBinned4::FPoolList::PushNewPoolToFront(FMalloc
 		OutOfMemory(GMB4PageSize);
 	}
 
-	FFreeBlock* Free = new (FreePtr) FFreeBlock(GMB4PageSize, Table.BinSize, InPoolIndex, Allocator.CurrentCanary);
+	FFreeBlock* Free = new (FreePtr) FFreeBlock(GMB4PageSize, Table.BinSize, Allocator.CurrentCanary);
 	check(IsAligned(Free, GMB4PageSize));
 
 #if UE_MB4_ALLOCATOR_STATS
@@ -366,6 +361,7 @@ FMallocBinned4::FPoolInfo& FMallocBinned4::FPoolList::PushNewPoolToFront(FMalloc
 #endif
 
 	UE::TUniqueLock Lock(Allocator.ExternalAllocMutex);
+	Allocator.RegisterBlockPool(FreePtr, InPoolIndex);
 	// Create pool
 	FPoolInfo* Result = Internal::GetOrCreatePoolInfo(Allocator, Free, FPoolInfo::ECanary::FirstFreeBlockIsPtr);
 	Result->Link(Front);
@@ -427,12 +423,23 @@ FMallocBinned4::FMallocBinned4()
 
 	AllocateHashBuckets();
 
+	BlockPoolL1Table = (uint8**)AllocateMetaDataMemory(65536 * sizeof(uint8*));
+	FMemory::Memzero(BlockPoolL1Table, 65536 * sizeof(uint8*));
+
 	MallocBinned4 = this;
 	GFixedMallocLocationPtr = (FMalloc**)(&MallocBinned4);
 }
 
 FMallocBinned4::~FMallocBinned4()
 {
+	for (int32 i = 0; i < 65536; ++i)
+	{
+		if (BlockPoolL1Table[i])
+		{
+			FreeMetaDataMemory(BlockPoolL1Table[i], 65536);
+		}
+	}
+	FreeMetaDataMemory(BlockPoolL1Table, 65536 * sizeof(uint8*));
 }
 
 void FMallocBinned4::OnMallocInitialized()
@@ -491,12 +498,30 @@ void FMallocBinned4::OnPostFork()
 		FMallocBinned4::Private::CheckThreadFreeBlockListsForFork();
 	}
 
-	// This will be compared against the pool header of existing allocations to turn Free into a no-op for pages shared with the parent process
 	UE_CLOG(CurrentCanary != EBlockCanary::PreFork, LogMemory, Fatal, TEXT("FMallocBinned4 only supports forking once!"));
 
 	OldCanary = CurrentCanary;
 	CurrentCanary = EBlockCanary::PostFork;
 
+	// Mark all currently registered small-pool blocks as pre-fork (bit 7) in the side table
+	{
+		UE::TUniqueLock Lock(ExternalAllocMutex);
+		for (int32 i = 0; i < 65536; ++i)
+		{
+			uint8* L2 = BlockPoolL1Table[i];
+			if (L2)
+			{
+				for (int32 j = 0; j < 65536; ++j)
+				{
+					if (L2[j] != UnregisteredBlock)
+					{
+						L2[j] |= PreForkFlag;
+					}
+				}
+			}
+		}
+	}
+
 	for (FPoolTable& Table : SmallPoolTables)
 	{
 		UE::TUniqueLock Lock(Table.Mutex);
@@ -661,19 +686,18 @@ void* FMallocBinned4::Realloc(void* Ptr, SIZE_T NewSize, uint32 Alignment)
 
 	checkf(FMath::IsPowerOfTwo(Alignment), TEXT("Invalid Realloc alignment: '%u' is not a power of two"), Alignment);
 
-	if (!IsOSAllocation(Ptr))
+	const uint8 ReallocEntry = GetSmallBinEntry(Ptr);
+	if (ReallocEntry != UnregisteredBlock)
 	{
-		check(Ptr); // null is 64k aligned so we should not be here
 		// Reallocate to a smaller/bigger pool if necessary
-		FFreeBlock* Free = GetPoolHeaderFromPointer(Ptr);
-		CanaryTest(Free);
-		const uint32 BinSize = Free->BinSize;
-		const uint32 PoolIndex = Free->PoolIndex;
+		const uint32 PoolIndex = ReallocEntry & PoolIndexMask;
+		CA_ASSUME(PoolIndex < UE_MB4_SMALL_POOL_COUNT);
+		const uint32 BinSize   = SmallPoolTables[PoolIndex].BinSize;
 #if BINNED4_FORK_SUPPORT
-		// If the canary is the pre-fork one, we should not allow this allocation to grow in-place to avoid copying a page from the parent process.
-		if (Free->CanaryAndForkState == CurrentCanary && 
+		// If this is a pre-fork page, we should not allow this allocation to grow in-place to avoid copying a page from the parent process.
+		if (!(ReallocEntry & PreForkFlag) &&
 #else
-		if( 
+		if(
 #endif
 			((NewSize <= BinSize) & (IsAligned(BinSize, Alignment))) && // one branch, not two
 			(PoolIndex == 0 || NewSize > PoolIndexToBinSize(PoolIndex - 1)))
@@ -752,21 +776,19 @@ void FMallocBinned4::Free(void* Ptr)
 {
 	NOALLOC_SCOPE_CYCLE_COUNTER(STAT_FMallocBinned4_FreeExternal);
 
-	if (!IsOSAllocation(Ptr))
+	const uint8 Entry = GetSmallBinEntry(Ptr);
+	if (Entry != UnregisteredBlock)
 	{
-		check(Ptr); // null is 64k aligned so we should not be here
-		FFreeBlock* BasePtr = GetPoolHeaderFromPointer(Ptr);
-		CanaryTest(BasePtr);
-
 #if BINNED4_FORK_SUPPORT
-		if (BasePtr->CanaryAndForkState != CurrentCanary)
+		if (Entry & PreForkFlag)
 		{
-			// This page was allocated before we forked so we want to avoid dirtying it by writing a linked list into it 
+			// This page was allocated before we forked so we want to avoid dirtying it by writing a linked list into it
 			return;
 		}
 #endif
-		const uint32 BinSize = BasePtr->BinSize;
-		const uint32 PoolIndex = BasePtr->PoolIndex;
+		const uint32 PoolIndex = Entry & PoolIndexMask;
+		CA_ASSUME(PoolIndex < UE_MB4_SMALL_POOL_COUNT);
+		const uint32 BinSize   = SmallPoolTables[PoolIndex].BinSize;
 
 		FBundleNode* BundlesToRecycle = nullptr;
 		FPerThreadFreeBlockLists* Lists = GMallocBinnedPerThreadCaches ? FPerThreadFreeBlockLists::Get() : nullptr;
@@ -886,15 +908,46 @@ void FMallocBinned4::FlushCurrentThreadCacheInternal(bool bNewEpochOnly)
 	FMallocBinnedCommonUtils::FlushCurrentThreadCache(*this, bNewEpochOnly);
 }
 
-void FMallocBinned4::CanaryFail(const FFreeBlock* Block) const
+void FMallocBinned4::CanaryFail(const void* Ptr) const
 {
 #if BINNED4_FORK_SUPPORT
-	UE_LOG(LogMemory, Fatal, TEXT("FMallocBinned4 Attempt to realloc an unrecognized block %p   canary == 0x%x != 0x%x or 0x%x "), (void*)Block, (int32)Block->CanaryAndForkState, CurrentCanary, OldCanary);
+	UE_LOG(LogMemory, Fatal, TEXT("FMallocBinned4 Attempt to free/realloc an unrecognized block %p  (no side table entry - expected pre-fork or post-fork registered block)"), Ptr);
 #else
-	UE_LOG(LogMemory, Fatal, TEXT("FMallocBinned4 Attempt to realloc an unrecognized block %p   canary == 0x%x != 0x%x"), (void*)Block, (int32)Block->CanaryAndForkState, int(CurrentCanary));
+	UE_LOG(LogMemory, Fatal, TEXT("FMallocBinned4 Attempt to free/realloc an unrecognized block %p  (no side table entry)"), Ptr);
 #endif
 }
 
+void FMallocBinned4::RegisterBlockPool(void* BlockBase, uint32 PoolIndex)
+{
+	const uint64 Addr = uint64(BlockBase);
+	uint8*& L2 = BlockPoolL1Table[Addr >> 32];
+	if (!L2)
+	{
+		L2 = (uint8*)AllocateMetaDataMemory(65536);
+		FMemory::Memset(L2, UnregisteredBlock, 65536);
+	}
+	L2[(Addr >> 16) & 0xFFFF] = uint8(PoolIndex);
+}
+
+void FMallocBinned4::UnregisterBlockPool(void* BlockBase)
+{
+	const uint64 Addr = uint64(BlockBase);
+	uint8* L2 = BlockPoolL1Table[Addr >> 32];
+	if (L2) L2[(Addr >> 16) & 0xFFFF] = UnregisteredBlock;
+}
+
+uint8 FMallocBinned4::GetSmallBinEntry(const void* Ptr) const
+{
+	const uint8 Entry = LookupBlockPoolEntry(Ptr);
+	if (!IsAligned(Ptr, UE_MB4_LARGE_ALLOC) || Entry != UnregisteredBlock)
+	{
+		check(Ptr); // null is 64KB-aligned so it should never reach this branch
+		if (Entry == UnregisteredBlock) CanaryFail(Ptr);
+		return Entry;
+	}
+	return UnregisteredBlock;
+}
+
 void FMallocBinned4::DumpAllocatorStats(class FOutputDevice& Ar)
 {
 #if UE_MB4_ALLOCATOR_STATS
diff --git a/Engine/Source/Runtime/Core/Public/HAL/MallocBinned4.h b/Engine/Source/Runtime/Core/Public/HAL/MallocBinned4.h
--- a/Engine/Source/Runtime/Core/Public/HAL/MallocBinned4.h
+++ b/Engine/Source/Runtime/Core/Public/HAL/MallocBinned4.h
@@ -27,7 +27,7 @@
 #	define UE_MB4_MAX_SMALL_POOL_SIZE		(13104)		// Maximum bin size in SmallBinSizesInternal in cpp file
 #	define UE_MB4_SMALL_POOL_COUNT			49
 #else
-#	define UE_MB4_MAX_SMALL_POOL_SIZE		(32768-16)	// Maximum bin size in SmallBinSizesInternal in cpp file
+#	define UE_MB4_MAX_SMALL_POOL_SIZE		(32768)		// Maximum bin size in SmallBinSizesInternal in cpp file
 #	define UE_MB4_SMALL_POOL_COUNT			52
 #endif
 
@@ -105,19 +105,12 @@ private:
 	/** Information about a piece of free memory. */
 	struct FFreeBlock
 	{
-		inline FFreeBlock(uint32 InPageSize, uint16 InBinSize, uint8 InPoolIndex, EBlockCanary InCanary)
+		inline FFreeBlock(uint32 InPageSize, uint16 InBinSize, EBlockCanary InCanary)
 			: BinSize(InBinSize)
-			, PoolIndex(InPoolIndex)
 			, CanaryAndForkState(InCanary)
 			, NextFreeBlockIndex(InvalidNextFreeBlock)
 		{
-			check(PoolIndex == InPoolIndex);
 			NumFreeBins = uint16(InPageSize / InBinSize);
-			if (NumFreeBins * InBinSize + sizeof(FFreeBlock) > InPageSize)
-			{
-				NumFreeBins--;
-			}
-			check(NumFreeBins * InBinSize + sizeof(FFreeBlock) <= InPageSize);
 		}
 
 		FORCEINLINE uint32 GetNumFreeBins() const
@@ -128,12 +121,7 @@ private:
 		inline void* AllocateBin()
 		{
 			--NumFreeBins;
-			// this is a pointer to a whole 64KB block - bins are allocated out from the end
-			if (IsAligned(this, UE_MB4_LARGE_ALLOC))
-			{
-				return (uint8*)this + UE_MB4_LARGE_ALLOC - (NumFreeBins + 1) * BinSize;
-			}
-			// this is a recycled single bin within the 64KB block
+			// allocate high-to-low: offset 0 (the header location) is returned last
 			return (uint8*)this + NumFreeBins * BinSize;
 		}
 
@@ -161,14 +149,11 @@ private:
 		}
 
 		//constexpr static uint32 MaxBitsForBinSize	 = 32 - FMath::CountLeadingZeros(UE_MB4_MAX_SMALL_POOL_SIZE));
-		constexpr static uint32 MaxBitsForPoolIndex = 32 - FMath::CountLeadingZeros(UE_MB4_SMALL_POOL_COUNT);
 		constexpr static uint32 MaxBitsForFreeBins = 32 - FMath::CountLeadingZeros(UE_MB4_LARGE_ALLOC / UE_MBC_MIN_BIN_SIZE);
 		constexpr static uint32 InvalidNextFreeBlock = 1u << MaxBitsForFreeBins;
-		static_assert(MaxBitsForPoolIndex <= 8, "PoolIndex should fit in 8 bits");
 		static_assert(MaxBitsForFreeBins + 1 <= 16, "NextFreeBlockIndex should fit in 16 bits");
 
 		uint16 BinSize;						// Size of the bins
-		uint8 PoolIndex;					// Index of this pool
 		EBlockCanary CanaryAndForkState;	// Canary value; two valid values when fork support is active
 		uint16 NumFreeBins;					// Number of consecutive free bins here, at least 1
 		uint16 NextFreeBlockIndex;			// Bin index within 64KB block of next recycled FFreeBlock, or InvalidNextFreeBlock
@@ -235,20 +220,40 @@ private:
 	FPooledVirtualMemoryAllocator CachedOSPageAllocator;
 #endif
 
+	// Side-table entry encoding: bits[6:0] = PoolIndex, bit 7 = pre-fork flag
+	static constexpr uint8 PreForkFlag       = 0x80;
+	static constexpr uint8 PoolIndexMask     = 0x7F;
+	static constexpr uint8 UnregisteredBlock = 0xFF;
+
+	uint8** BlockPoolL1Table = nullptr;
+
+	FORCEINLINE uint8 LookupBlockPoolEntry(const void* Ptr) const
+	{
+		const uint64 Addr = uint64(AlignDown(Ptr, UE_MB4_LARGE_ALLOC));
+		const uint8* L2 = BlockPoolL1Table[Addr >> 32];
+		return L2 ? L2[(Addr >> 16) & 0xFFFF] : UnregisteredBlock;
+	}
+
+	void RegisterBlockPool(void* BlockBase, uint32 PoolIndex);
+	void UnregisterBlockPool(void* BlockBase);
+
+	// Returns the validated side-table entry for a small-pool allocation,
+	// or UnregisteredBlock if Ptr is a large OS allocation (64KB-aligned and unregistered) or null.
+	// Calls CanaryFail if Ptr looks like a small-pool pointer but has no registered entry.
+	uint8 GetSmallBinEntry(const void* Ptr) const;
+
 	FORCEINLINE bool IsOSAllocation(const void* Ptr) const
 	{
 #if UE_USE_VERYLARGEPAGEALLOCATOR && !PLATFORM_UNIX && !PLATFORM_ANDROID
-		return !CachedOSPageAllocator.IsSmallBlockAllocation(Ptr) && IsAligned(Ptr, UE_MB4_LARGE_ALLOC);
+		// With the very large page allocator, small pool blocks come from managed 2MB pages
+		// and must be distinguished from large OS allocations via IsSmallBlockAllocation.
+		// The side table provides the same distinction but we keep both checks for correctness.
+		return !CachedOSPageAllocator.IsSmallBlockAllocation(Ptr) && IsAligned(Ptr, UE_MB4_LARGE_ALLOC) && LookupBlockPoolEntry(Ptr) == UnregisteredBlock;
 #else
-		return IsAligned(Ptr, UE_MB4_LARGE_ALLOC);
+		return IsAligned(Ptr, UE_MB4_LARGE_ALLOC) && LookupBlockPoolEntry(Ptr) == UnregisteredBlock;
 #endif
 	}
 
-	static FORCEINLINE FFreeBlock* GetPoolHeaderFromPointer(void* Ptr)
-	{
-		return (FFreeBlock*)AlignDown(Ptr, UE_MB4_LARGE_ALLOC);
-	}
-
 public:
 	FMallocBinned4();
 	virtual ~FMallocBinned4();
@@ -267,14 +272,12 @@ public:
 
 	inline bool GetSmallAllocationSize(void* Ptr, SIZE_T& SizeOut) const
 	{
-		if (!IsOSAllocation(Ptr))
-		{
-			const FFreeBlock* Free = GetPoolHeaderFromPointer(Ptr);
-			CanaryTest(Free);
-			SizeOut = Free->BinSize;
-			return true;
-		}
-		return false;
+		const uint8 Entry = LookupBlockPoolEntry(Ptr);
+		if (Entry == UnregisteredBlock) return false;
+		const uint32 PoolIndex = Entry & PoolIndexMask;
+		CA_ASSUME(PoolIndex < UE_MB4_SMALL_POOL_COUNT);
+		SizeOut = SmallPoolTables[PoolIndex].BinSize;
+		return true;
 	}
 
 	inline virtual bool GetAllocationSize(void *Ptr, SIZE_T &SizeOut) override
@@ -311,17 +314,17 @@ public:
 	void* MallocExternalSmall(SIZE_T Size, uint32 Alignment);
 	void* MallocExternalLarge(SIZE_T Size, uint32 Alignment);
 
-	void CanaryFail(const FFreeBlock* Block) const;
-	inline void CanaryTest(const FFreeBlock* Block) const
+	void CanaryFail(const void* Ptr) const;
+	inline void CanaryTest(const void* Ptr) const
 	{
-#if BINNED4_FORK_SUPPORT
-		// When we support forking there are two valid canary values.
-		if (Block->CanaryAndForkState != CurrentCanary && Block->CanaryAndForkState != OldCanary)
-#else
-		if (Block->CanaryAndForkState != CurrentCanary)
-#endif
+		// The side table is used instead of reading FFreeBlock::CanaryAndForkState because
+		// the new AllocateBin direction returns offset 0 (the header location) as the last bin,
+		// so the block header may be user-owned memory at the time of Free/Realloc.
+		// Pre-fork blocks remain registered in the side table (with bit7 set) and therefore
+		// pass this test, matching the original two-canary BINNED4_FORK_SUPPORT behavior.
+		if (LookupBlockPoolEntry(Ptr) == UnregisteredBlock)
 		{
-			CanaryFail(Block);
+			CanaryFail(Ptr);
 		}
 	}
 
