Added versioned bHasPayload flag and framework for dynamic delegate serialization.

#Synced-CL 52013254

diff --git a/Engine/Source/Runtime/Core/Private/UObject/ScriptDelegates.cpp b/Engine/Source/Runtime/Core/Private/UObject/ScriptDelegates.cpp
new file mode 100644
--- /dev/null
+++ b/Engine/Source/Runtime/Core/Private/UObject/ScriptDelegates.cpp
@@ -0,0 +1,15 @@
+// Copyright Epic Games, Inc. All Rights Reserved.
+
+#include "UObject/ScriptDelegates.h"
+
+#if UE_USE_DYNAMIC_DELEGATE_PAYLOADS
+
+namespace UE::Core::Private
+{
+	FDelegatePayloadSerializeFunc       GDelegatePayloadSerializeFunc       = nullptr;
+	FDelegatePayloadDeserializeFunc     GDelegatePayloadDeserializeFunc     = nullptr;
+	FDelegatePayloadSerializeSlotFunc   GDelegatePayloadSerializeSlotFunc   = nullptr;
+	FDelegatePayloadDeserializeSlotFunc GDelegatePayloadDeserializeSlotFunc = nullptr;
+}
+
+#endif
diff --git a/Engine/Source/Runtime/Core/Public/UObject/FortniteMainBranchObjectVersions.inl b/Engine/Source/Runtime/Core/Public/UObject/FortniteMainBranchObjectVersions.inl
--- a/Engine/Source/Runtime/Core/Public/UObject/FortniteMainBranchObjectVersions.inl
+++ b/Engine/Source/Runtime/Core/Public/UObject/FortniteMainBranchObjectVersions.inl
@@ -796,5 +796,8 @@ FORTNITE_MAIN_VERSION(AddDevNotesToFText, 260)
 // Added default color override for PCG Mesh Sampler and fix bounds in PCG Subdivide Spline
 FORTNITE_MAIN_VERSION(PCGMeshSamplerDefaultColorPCGSubdivideSplineBoundsFixes, 261)
 
+// Serialisable dynamic delegate payloads
+FORTNITE_MAIN_VERSION(DynamicDelegatePayloads, 262)
+
 // -----<new versions can be added above this line>-------------------------------------------------
 #undef FORTNITE_MAIN_VERSION
diff --git a/Engine/Source/Runtime/Core/Public/UObject/ScriptDelegates.h b/Engine/Source/Runtime/Core/Public/UObject/ScriptDelegates.h
--- a/Engine/Source/Runtime/Core/Public/UObject/ScriptDelegates.h
+++ b/Engine/Source/Runtime/Core/Public/UObject/ScriptDelegates.h
@@ -17,6 +17,7 @@
 #include "Templates/SharedPointer.h"
 #include "Templates/TypeHash.h"
 #include "Templates/UnrealTypeTraits.h"
+#include "UObject/FortniteMainBranchObjectVersion.h"
 #include "UObject/NameTypes.h"
 #include <type_traits>
 
@@ -155,9 +156,13 @@ namespace UE::Core::Private
 	template <typename WeakPtrType>
 	struct IDynamicDelegatePayload
 	{
-		// Calls ObjectPtr->ProcessEvent(Function, ParametersPlusPayload...) on ObjectPtr. 
+		// Calls ObjectPtr->ProcessEvent(Function, ParametersPlusPayload...) on ObjectPtr.
 		virtual void ProcessEvent(WeakPtrType ObjectPtr, UFunction* Function, void* Parameters) const = 0;
 
+		// Serializes the payload in binary form - only intended for saving, so const.
+		// Payloads are always binary, even within structured text archives.
+		virtual void Serialize(FArchive& Ar) const = 0;
+
 		virtual ~IDynamicDelegatePayload() = 0;
 	};
 
@@ -213,6 +218,14 @@ namespace UE::Core::Private
 			}
 		}
 
+		void Serialize(FArchive& Ar) const override
+		{
+			// Function is const because it's only expected to be used for saving.  The Payload is mutable so this is ok.
+			checkSlow(!Ar.IsLoading());
+
+			// TODO - implement serialization.
+		}
+
 	private:
 		// Mutable because ProcessEvent takes a non-const param pointer
 		mutable VarsTupleType Payload;
@@ -223,6 +236,22 @@ namespace UE::Core::Private
 	{
 		return MakeShared<TDynamicDelegatePayload<WeakPtrType, BoundFuncType, TDynamicDelegateParamPayloadCaptureType_T<WeakPtrType, VarTypes>...>>(Forward<VarTypes>(Vars)...);
 	}
+
+	/** Callback function pointers for serializing/deserializing delegate payloads.
+	*  Null until CoreUObject registers them at module startup (Step 4).
+	*  All payload serialization is implemented in CoreUObject — Core only holds
+	*  an opaque IDynamicDelegatePayload handle without knowing the details.
+	*  Payloads are always serialized in binary form, even within structured text
+	*  archives (where they appear as a base-64 blob via FArchiveUObjectFromStructuredArchive). */
+	using FDelegatePayloadSerializeFunc      = void(*)(FArchive& Ar, const IDynamicDelegatePayload<FWeakObjectPtr>& Payload);
+	using FDelegatePayloadDeserializeFunc    = TSharedPtr<IDynamicDelegatePayload<FWeakObjectPtr>>(*)(FArchive& Ar);
+	using FDelegatePayloadSerializeSlotFunc  = void(*)(FStructuredArchiveSlot Slot, const IDynamicDelegatePayload<FWeakObjectPtr>& Payload);
+	using FDelegatePayloadDeserializeSlotFunc = TSharedPtr<IDynamicDelegatePayload<FWeakObjectPtr>>(*)(FStructuredArchiveSlot Slot);
+
+	CORE_API extern FDelegatePayloadSerializeFunc       GDelegatePayloadSerializeFunc;
+	CORE_API extern FDelegatePayloadDeserializeFunc     GDelegatePayloadDeserializeFunc;
+	CORE_API extern FDelegatePayloadSerializeSlotFunc   GDelegatePayloadSerializeSlotFunc;
+	CORE_API extern FDelegatePayloadDeserializeSlotFunc GDelegatePayloadDeserializeSlotFunc;
 #endif
 
 	// Describes a type which is constructible with a ForceInit argument.
@@ -585,7 +614,46 @@ public:
 		FReadAccessScope ReadScope = D.GetReadAccessScope();
 
 		Ar << D.Object << D.FunctionName;
-		// TODO: Serialize payload
+
+		Ar.UsingCustomVersion(FFortniteMainBranchObjectVersion::GUID);
+
+		bool bHasPayload = false;
+		if (Ar.CustomVer(FFortniteMainBranchObjectVersion::GUID) >= FFortniteMainBranchObjectVersion::DynamicDelegatePayloads)
+		{
+#if UE_USE_DYNAMIC_DELEGATE_PAYLOADS
+			bHasPayload = !!D.Payload;
+#endif
+			Ar << bHasPayload;
+		}
+
+		if (bHasPayload)
+		{
+			if (Ar.IsLoading())
+			{
+#if UE_USE_DYNAMIC_DELEGATE_PAYLOADS
+				if (ensureMsgf(UE::Core::Private::GDelegatePayloadDeserializeFunc, TEXT("Trying to deserialize a dynamic delegate payload without CoreUObject's deserialize factory being registered")))
+				{
+					D.Payload = UE::Core::Private::GDelegatePayloadDeserializeFunc(Ar);
+				}
+				else
+				{
+					D.Unbind();
+					// TODO: drain payload bytes once FDynamicDelegatePropertyPayload exists (Step 3)
+				}
+#else
+				ensureMsgf(false, TEXT("Trying to deserialize a dynamic delegate payload when payloads are disabled"));
+				D.Unbind();
+				// TODO: drain bytes (Step 3)
+#endif
+			}
+#if UE_USE_DYNAMIC_DELEGATE_PAYLOADS
+			else
+			{
+				UE::Core::Private::GDelegatePayloadSerializeFunc(Ar, *D.Payload);
+			}
+#endif
+		}
+
 		return Ar;
 	}
 
@@ -594,9 +662,52 @@ public:
 	{
 		FReadAccessScope ReadScope = D.GetReadAccessScope();
 
+		FArchive& UnderlyingAr = Slot.GetUnderlyingArchive();
+
 		FStructuredArchive::FRecord Record = Slot.EnterRecord();
 		Record << SA_VALUE(TEXT("Object"), D.Object) << SA_VALUE(TEXT("FunctionName"),D.FunctionName);
-		// TODO: Serialize payload
+
+		UnderlyingAr.UsingCustomVersion(FFortniteMainBranchObjectVersion::GUID);
+
+		if (UnderlyingAr.CustomVer(FFortniteMainBranchObjectVersion::GUID) >= FFortniteMainBranchObjectVersion::DynamicDelegatePayloads)
+		{
+			bool bHasPayload = false;
+#if UE_USE_DYNAMIC_DELEGATE_PAYLOADS
+			bHasPayload = !!D.Payload;
+#endif
+
+			// TryEnterField writes the Payload field only when bHasPayload is true (saving),
+			// and on load returns a slot only if the field exists in the archive.
+			// The payload is always binary — the slot hooks use FArchiveUObjectFromStructuredArchive
+			// (in CoreUObject) to wrap the slot as a binary archive, producing a base-64 blob in text formats.
+			if (TOptional<FStructuredArchiveSlot> PayloadSlot = Record.TryEnterField(TEXT("Payload"), bHasPayload))
+			{
+				if (UnderlyingAr.IsLoading())
+				{
+#if UE_USE_DYNAMIC_DELEGATE_PAYLOADS
+					if (ensureMsgf(UE::Core::Private::GDelegatePayloadDeserializeSlotFunc, TEXT("Trying to deserialize a dynamic delegate payload without CoreUObject's deserialize factory being registered")))
+					{
+						D.Payload = UE::Core::Private::GDelegatePayloadDeserializeSlotFunc(*PayloadSlot);
+					}
+					else
+					{
+						D.Unbind();
+						// TODO: drain payload bytes once FDynamicDelegatePropertyPayload exists (Step 3)
+					}
+#else
+					ensureMsgf(false, TEXT("Trying to deserialize a dynamic delegate payload when payloads are disabled"));
+					D.Unbind();
+					// TODO: drain bytes (Step 3)
+#endif
+				}
+#if UE_USE_DYNAMIC_DELEGATE_PAYLOADS
+				else
+				{
+					UE::Core::Private::GDelegatePayloadSerializeSlotFunc(*PayloadSlot, *D.Payload);
+				}
+#endif
+			}
+		}
 	}
 
 	/** Comparison operators */
@@ -1176,7 +1287,7 @@ public:
 				// Create a copy of the invocation list, just in case the list is modified by one of the callbacks during the broadcast
 				typedef TArray< UnicastDelegateType, TInlineAllocator< 4 > > FInlineInvocationList;
 				FInlineInvocationList InvocationListCopy = FInlineInvocationList(InvocationList);
-		
+
 				// Invoke each bound function
 				for( typename FInlineInvocationList::TConstIterator FunctionIt( InvocationListCopy ); FunctionIt; ++FunctionIt )
 				{
@@ -1229,7 +1340,7 @@ public:
 	/**
 	 * Returns all objects associated with this multicast-delegate, even if unreachable.
 	 * For advanced uses only -- you should never need call this function in normal circumstances.
- 	 * @return	List of objects bound to this delegate
+	 * @return	List of objects bound to this delegate
 	 */
 	TArray<UObject*> GetAllObjectsEvenIfUnreachable() const
 	{
@@ -1280,7 +1391,7 @@ protected:
 	 * Adds a function delegate to this multi-cast delegate's invocation list
 	 *
 	 * @param	InDelegate	Delegate to add
-	*/
+	 */
 	void AddInternal(UnicastDelegateType&& InDelegate)
 	{
 #if DO_ENSURE
@@ -1311,7 +1422,7 @@ protected:
 	 * order of the delegates may not be preserved!
 	 *
 	 * @param	InDelegate	Delegate to remove
-	*/
+	 */
 	void RemoveInternal( const UnicastDelegateType& InDelegate ) const
 	{
 		InvocationList.RemoveSingleSwap(InDelegate);
@@ -1323,7 +1434,7 @@ protected:
 	 *
 	 * @param	InObject		Object of the delegate to remove
 	 * @param	InFunctionName	Function name of the delegate to remove
-	*/
+	 */
 	void RemoveInternal( const UObject* InObject, FName InFunctionName ) const
 	{
 		int32 FoundDelegate = InvocationList.IndexOfByPredicate([=](const UnicastDelegateType& Delegate) {
