Added MB3 support to Unix platform.
It can be enabled via a cmd line -binnedmalloc3 or via StaticAllocator setting in project's Target.cs file

#Synced-CL 52054134

diff --git a/Engine/Source/Runtime/Core/Private/Unix/UnixPlatformMemory.cpp b/Engine/Source/Runtime/Core/Private/Unix/UnixPlatformMemory.cpp
--- a/Engine/Source/Runtime/Core/Private/Unix/UnixPlatformMemory.cpp
+++ b/Engine/Source/Runtime/Core/Private/Unix/UnixPlatformMemory.cpp
@@ -19,6 +19,7 @@
 #include "HAL/MallocJemalloc.h"
 #include "HAL/MallocBinned.h"
 #include "HAL/MallocBinned2.h"
+#include "HAL/MallocBinned3.h"
 #include "HAL/MallocReplayProxy.h"
 #include "HAL/MallocStomp.h"
 #include "HAL/PlatformMallocCrash.h"
@@ -158,6 +159,10 @@ class FMalloc* FUnixPlatformMemory::BaseAllocator()
 	{
 		AllocatorToUse = EMemoryAllocatorToUse::Binned2;
 	}
+	else if (USE_MALLOC_BINNED3)
+	{
+		AllocatorToUse = EMemoryAllocatorToUse::Binned3;
+	}
 	else
 	{
 		AllocatorToUse = EMemoryAllocatorToUse::Binned;
@@ -212,6 +217,12 @@ class FMalloc* FUnixPlatformMemory::BaseAllocator()
 				break;
 			}
 
+			if (FCStringAnsi::Stricmp(Arg, "-binnedmalloc3") == 0)
+			{
+				AllocatorToUse = EMemoryAllocatorToUse::Binned3;
+				break;
+			}
+
 			if (FCStringAnsi::Stricmp(Arg, "-fullcrashcallstack") == 0)
 			{
 				GFullCrashCallstack = true;
@@ -338,6 +349,10 @@ class FMalloc* FUnixPlatformMemory::BaseAllocator()
 		Allocator = new FMallocBinned2();
 		break;
 
+	case EMemoryAllocatorToUse::Binned3:
+		Allocator = new FMallocBinned3();
+		break;
+
 	default:	// intentional fall-through
 	case EMemoryAllocatorToUse::Binned:
 		Allocator = new FMallocBinned(FPlatformMemory::GetConstants().BinnedPageSize & MAX_uint32, 0x100000000);
