Provided a hint to static analysis to silence a warning in MB4

#Synced-CL 52054134

diff --git a/Engine/Source/Runtime/Core/Private/HAL/MallocBinned4.cpp b/Engine/Source/Runtime/Core/Private/HAL/MallocBinned4.cpp
--- a/Engine/Source/Runtime/Core/Private/HAL/MallocBinned4.cpp
+++ b/Engine/Source/Runtime/Core/Private/HAL/MallocBinned4.cpp
@@ -707,6 +707,7 @@ void* FMallocBinned4::Realloc(void* Ptr, SIZE_T NewSize, uint32 Alignment)
 
 		// Reallocate and copy the data across
 		void* Result = FMallocBinned4::Malloc(NewSize, Alignment);
+		CA_ASSUME(Ptr);
 		FMemory::Memcpy(Result, Ptr, FMath::Min<SIZE_T>(NewSize, BinSize));
 		FMallocBinned4::Free(Ptr);
 		return Result;
