Add frame pointer-based stack walker for Linux memory trace callstacks to improve Unreal Insight memory capture speed

#rnx
#Synced-CL 52169359

diff --git a/Engine/Source/Runtime/Core/Core.Build.cs b/Engine/Source/Runtime/Core/Core.Build.cs
--- a/Engine/Source/Runtime/Core/Core.Build.cs
+++ b/Engine/Source/Runtime/Core/Core.Build.cs
@@ -216,6 +216,7 @@ public class Core : ModuleRules
 				PublicDefinitions.Add("UE_MEMORY_TRACE_AVAILABLE=1");
 				PublicDefinitions.Add("UE_MEMORY_TAGS_TRACE_ENABLED=1");
 				PublicDefinitions.Add("UE_CALLSTACK_TRACE_ENABLED=1");
+				PrivateDefinitions.Add("UE_CALLSTACK_TRACE_UNIX_USE_STACK_FRAMES_WALKING=1");
 			}
 		}
 
diff --git a/Engine/Source/Runtime/Core/Private/ProfilingDebugging/Unix/UnixCallstackTrace.cpp b/Engine/Source/Runtime/Core/Private/ProfilingDebugging/Unix/UnixCallstackTrace.cpp
--- a/Engine/Source/Runtime/Core/Private/ProfilingDebugging/Unix/UnixCallstackTrace.cpp
+++ b/Engine/Source/Runtime/Core/Private/ProfilingDebugging/Unix/UnixCallstackTrace.cpp
@@ -4,6 +4,10 @@
 
 #if UE_CALLSTACK_TRACE_ENABLED
 
+#ifndef UE_CALLSTACK_TRACE_UNIX_USE_STACK_FRAMES_WALKING
+	#define UE_CALLSTACK_TRACE_UNIX_USE_STACK_FRAMES_WALKING 0
+#endif
+
 #include "CoreTypes.h"
 #include "HAL/Platform.h"
 #include "HAL/PlatformStackWalk.h"
@@ -46,12 +50,31 @@ FBacktracer* FBacktracer::Get()
 	return Instance;
 }
 
+////////////////////////////////////////////////////////////////////////////////
+#if UE_CALLSTACK_TRACE_UNIX_USE_STACK_FRAMES_WALKING
+// Minimum frame count to consider frame pointer walking successful. A real
+// call from the memory allocator should produce a deep stack. Fewer frames
+// means frame pointers are likely not available, so fall back to backtrace().
+static constexpr uint32 GMinFramesForFramePointerWalking = 5;
+#endif
+
 ////////////////////////////////////////////////////////////////////////////////
 uint32 FBacktracer::GetBacktraceId(void* ReturnAddress)
 {
 #if !UE_BUILD_SHIPPING
 	uint64 StackFrames[UE_CALLSTACK_TRACE_MAX_FRAMES];
+#if UE_CALLSTACK_TRACE_UNIX_USE_STACK_FRAMES_WALKING
+	// Try the fast frame pointer walker first. If it doesn't return enough
+	// frames (frame pointers not available), fall back to the slower but
+	// more thorough DWARF-based backtrace().
+	uint32 NumStackFrames = FPlatformStackWalk::CaptureStackBackTraceViaFramePointerWalking(StackFrames, UE_ARRAY_COUNT(StackFrames));
+	if (NumStackFrames < GMinFramesForFramePointerWalking)
+	{
+		NumStackFrames = FPlatformStackWalk::CaptureStackBackTrace(StackFrames, UE_ARRAY_COUNT(StackFrames));
+	}
+#else
 	int32 NumStackFrames = FPlatformStackWalk::CaptureStackBackTrace(StackFrames, UE_ARRAY_COUNT(StackFrames));
+#endif
 	if (NumStackFrames > 0)
 	{
 		FCallstackTracer::FBacktraceEntry BacktraceEntry;
diff --git a/Engine/Source/Runtime/Core/Private/Unix/UnixPlatformStackWalk.cpp b/Engine/Source/Runtime/Core/Private/Unix/UnixPlatformStackWalk.cpp
--- a/Engine/Source/Runtime/Core/Private/Unix/UnixPlatformStackWalk.cpp
+++ b/Engine/Source/Runtime/Core/Private/Unix/UnixPlatformStackWalk.cpp
@@ -17,6 +17,7 @@
 #include "HAL/PlatformTime.h"
 #include "AutoRTFM.h"
 
+#include <atomic>
 #include <link.h>
 #include <signal.h>
 
@@ -801,6 +802,97 @@ UE_AUTORTFM_ALWAYS_OPEN uint32 FUnixPlatformStackWalk::CaptureStackBackTrace( ui
 	return (uint32)Size;
 }
 
+UE_AUTORTFM_ALWAYS_OPEN uint32 FUnixPlatformStackWalk::CaptureStackBackTraceViaFramePointerWalking(uint64* BackTrace, uint32 MaxDepth)
+{
+#if defined(__aarch64__) || defined(__x86_64__)
+	uint64 StackTopPtr = 0;
+
+	// Cache the main thread's stack top to avoid repeated pthread_getattr_np calls.
+	const bool bIsMainThread = syscall(SYS_gettid) == getpid();
+	static std::atomic<uint64> StackTopMainThread{0};
+	uint64 CachedStackTop = StackTopMainThread.load(std::memory_order_relaxed);
+	if (bIsMainThread && CachedStackTop != 0)
+	{
+		StackTopPtr = CachedStackTop;
+	}
+	else
+	{
+		pthread_attr_t ThreadAttr;
+		if (pthread_getattr_np(pthread_self(), &ThreadAttr) != 0)
+		{
+			return 0;
+		}
+
+		void* StackBase;
+		size_t StackSize;
+		if (pthread_attr_getstack(&ThreadAttr, &StackBase, &StackSize) != 0)
+		{
+			pthread_attr_destroy(&ThreadAttr);
+			return 0;
+		}
+
+		pthread_attr_destroy(&ThreadAttr);
+
+		StackTopPtr = (uint64)StackBase + StackSize;
+
+		if (bIsMainThread)
+		{
+			StackTopMainThread.store(StackTopPtr, std::memory_order_relaxed);
+		}
+	}
+
+	struct StackFrame
+	{
+		StackFrame* NextFrame;
+		UPTRINT ReturnPtr;
+
+		inline UPTRINT GetReturnPtr()
+		{
+#if defined(__aarch64__)
+			uintptr_t Ptr = ReturnPtr;
+			// xpaclri operates on the link register (x30), so move value there and back.
+			// This instruction is HINT #7, mapped to NOP on pre-PAC architectures.
+			asm volatile(
+				"mov x30, %0\n"
+				"xpaclri\n"
+				"mov %0, x30\n"
+				: "+r"(Ptr)
+				:
+				: "x30"
+			);
+			return Ptr;
+#else
+			return ReturnPtr;
+#endif
+		}
+	};
+
+	StackFrame* FrameStart = (StackFrame*)__builtin_frame_address(0);
+	if (FrameStart == nullptr)
+	{
+		return 0;
+	}
+
+	StackFrame* FrameEnd = (StackFrame*)StackTopPtr;
+	uint32 NumStackFrames = 0;
+
+	for (StackFrame* CurrentFrame = FrameStart;
+		NumStackFrames < MaxDepth &&
+		CurrentFrame->NextFrame > FrameStart &&
+		CurrentFrame->NextFrame < FrameEnd &&
+		((UPTRINT)CurrentFrame->NextFrame & (sizeof(StackFrame) - 1)) == 0 && // stop at unaligned frame
+		CurrentFrame->GetReturnPtr() != 0; // stop if function ptr is 0
+		CurrentFrame = CurrentFrame->NextFrame, NumStackFrames++)
+	{
+		BackTrace[NumStackFrames] = CurrentFrame->GetReturnPtr();
+	}
+
+	return NumStackFrames;
+#else
+	return 0;
+#endif
+}
+
 namespace
 {
 	// Shorter timeout for crash handling — the default CVar is 60s which would cause
diff --git a/Engine/Source/Runtime/Core/Public/Unix/UnixPlatformStackWalk.h b/Engine/Source/Runtime/Core/Public/Unix/UnixPlatformStackWalk.h
--- a/Engine/Source/Runtime/Core/Public/Unix/UnixPlatformStackWalk.h
+++ b/Engine/Source/Runtime/Core/Public/Unix/UnixPlatformStackWalk.h
@@ -16,6 +16,9 @@ struct FUnixPlatformStackWalk : public FGenericPlatformStackWalk
 	static CORE_API void ProgramCounterToSymbolInfo( uint64 ProgramCounter, FProgramCounterSymbolInfo& out_SymbolInfo );
 	static CORE_API bool ProgramCounterToHumanReadableString( int32 CurrentCallDepth, uint64 ProgramCounter, ANSICHAR* HumanReadableString, SIZE_T HumanReadableStringSize, FGenericCrashContext* Context = nullptr );
 	static CORE_API uint32 CaptureStackBackTrace( uint64* BackTrace, uint32 MaxDepth, void* Context = nullptr );
+	// Fast stack backtrace, only for tracing.
+	// Will NOT work from signal handlers, is NOT suitable for crashreporting. Will not see inline functions.
+	static CORE_API uint32 CaptureStackBackTraceViaFramePointerWalking(uint64* BackTrace, uint32 MaxDepth);
 	static CORE_API void StackWalkAndDump( ANSICHAR* HumanReadableString, SIZE_T HumanReadableStringSize, int32 IgnoreCount, void* Context = nullptr );
 	static CORE_API void StackWalkAndDump( ANSICHAR* HumanReadableString, SIZE_T HumanReadableStringSize, void* ProgramCounter, void* Context = nullptr );
 	static CORE_API void StackWalkAndDumpEx(ANSICHAR* HumanReadableString, SIZE_T HumanReadableStringSize, int32 IgnoreCount, uint32 Flags, void* Context = nullptr);
