[Backout] - CL52172259
#rnx
Original CL Desc
-----------------------------------------------------------------
Added runtime network status monitoring to platforms missing it (iOS, Android, Windows, Mac, Linux)

diff --git a/Engine/Source/Runtime/Core/Private/Apple/ApplePlatformMisc.cpp b/Engine/Source/Runtime/Core/Private/Apple/ApplePlatformMisc.cpp
--- a/Engine/Source/Runtime/Core/Private/Apple/ApplePlatformMisc.cpp
+++ b/Engine/Source/Runtime/Core/Private/Apple/ApplePlatformMisc.cpp
@@ -560,20 +560,17 @@ void FApplePlatformMisc::PlatformInit()
 			}
 			UE::FApplePlatformMisc::CachedNetworkPath = InPath;
 			nw_retain(UE::FApplePlatformMisc::CachedNetworkPath);
-
-			ENetworkConnectionType ConnectionType = UE::FApplePlatformMisc::CheckNetworkConnectionType();
-
 			if (FTaskGraphInterface::IsRunning())
 			{
+				ENetworkConnectionType ConnectionType = UE::FApplePlatformMisc::CheckNetworkConnectionType();
 				if (ConnectionType != UE::FApplePlatformMisc::LastReportedNetworkType)
 				{
 					UE::FApplePlatformMisc::LastReportedNetworkType = ConnectionType;
 					FFunctionGraphTask::CreateAndDispatchWhenReady([ConnectionType]() {
 						FCoreDelegates::OnNetworkConnectionChanged.Broadcast(ConnectionType);
-						FPlatformMisc::SetNetworkConnectionStatus(FPlatformMisc::NetworkConnectionTypeToStatus(ConnectionType));
 					}, TStatId(), nullptr, ENamedThreads::GameThread);
 				}
-
+				
 				// Adding a reference, so path survives to the task body.
 				nw_retain(InPath);
 				FFunctionGraphTask::CreateAndDispatchWhenReady(
@@ -583,18 +580,10 @@ void FApplePlatformMisc::PlatformInit()
 					UE::FApplePlatformMisc::NotifyAboutNetworkConnectionCharacteristicsChange(InPath);
 				}, TStatId());
 			}
-			else
-			{
-				// Task graph not running yet (early startup). Set the initial status synchronously so it isn't silently dropped.
-				UE::FApplePlatformMisc::LastReportedNetworkType = ConnectionType;
-				FPlatformMisc::SetNetworkConnectionStatus(FPlatformMisc::NetworkConnectionTypeToStatus(ConnectionType));
-			}
 		}
 	);
 	
 	nw_path_monitor_set_queue(UE::FApplePlatformMisc::PathMonitor, dispatch_get_main_queue());
 	nw_path_monitor_start(UE::FApplePlatformMisc::PathMonitor);
 
-	// The nw_path_monitor callback will fire shortly after start.
-	// If the task graph is not yet running at that point, the callback sets the initial network status synchronously
 }
