TArray AllocatorInstance: Add constructors to TArray to allow passing in an AllocatorInstance RValue to move-construct the TArray's AllocatorInstance.
ContainerAllocator:s Create TSizedMallocAllocator to allow TArray to allocate from a runtime-provided FMalloc pointer.
ContainerAllocators: Move TSizedAllocator base functionality from TSizedHeapAllocator into a new base class TSizedAllocatorBase, and create the two subclasses for the existing functionality - TSizedHeapAllocator and TSizedExternalAllocator - and make TSizedMallocAllocator a third subclass of the new base class.
#rnx
#Synced-CL 52123874

diff --git a/Engine/Source/Runtime/Core/Private/FramePro/FrameProProfiler.cpp b/Engine/Source/Runtime/Core/Private/FramePro/FrameProProfiler.cpp
--- a/Engine/Source/Runtime/Core/Private/FramePro/FrameProProfiler.cpp
+++ b/Engine/Source/Runtime/Core/Private/FramePro/FrameProProfiler.cpp
@@ -239,7 +239,7 @@ class FFrameProProfilerContext : public TThreadSingleton<FFrameProProfilerContex
 	}
 	
 	/** Array that represents thread of scopes */
-	TArray<FFrameProProfilerScope, TSizedHeapAllocator<32, FrameProMalloc>> ProfilerScopes;
+	TArray<FFrameProProfilerScope, TSizedExternalAllocator<32, FrameProMalloc>> ProfilerScopes;
 
 public:
 
diff --git a/Engine/Source/Runtime/Core/Public/Containers/Array.h b/Engine/Source/Runtime/Core/Public/Containers/Array.h
--- a/Engine/Source/Runtime/Core/Public/Containers/Array.h
+++ b/Engine/Source/Runtime/Core/Public/Containers/Array.h
@@ -804,6 +804,24 @@ public:
 	{
 	}
 
+	/**
+	 * Constructor to create an empty array, that also initializes the ElementAllocatorType from a moved-in parameter.
+	 * The input InAllocatorInstance must be empty; it is only meant to pass parameters to the AllocatorInstance.
+	 * 
+	 * This constructor and the others that take ElementAllocatorType&& are intended to allow AllocatorTypes that need
+	 * to be parameterized, such as one that allocates from a specified FMalloc*.
+	 * The input AllocatorInstance must be empty because we currently do not have a way to store the correct ArrayMax
+	 * for the Allocator's current allocation; that would require adding an additional InArrayMax argument or extending
+	 * the API of AllocatorInstance to be able to report it.
+	 */
+	[[nodiscard]] UE_NODEBUG UE_FORCEINLINE_HINT explicit TArray(ElementAllocatorType&& InAllocatorInstance)
+		: AllocatorInstance(MoveTemp(InAllocatorInstance))
+		, ArrayNum(0)
+		, ArrayMax(AllocatorInstance.GetInitialCapacity())
+	{
+		check(!AllocatorInstance.HasAllocation());
+	}
+
 	/** Explicit consteval constructor for allocators that require zero-initialization of members for constant initialization */
 	[[nodiscard]] explicit consteval TArray(EConstEval)
 		: AllocatorInstance(ConstEval)
@@ -820,6 +838,23 @@ public:
 	 * @see Append
 	 */
 	[[nodiscard]] UE_FORCEINLINE_HINT TArray(const ElementType* Ptr, SizeType Count)
+	{
+		Construct(Ptr, Count);
+	}
+
+	/**
+	 * Constructor from a raw array of elements, that also initializes the ElementAllocatorType from a moved-in parameter.
+	 * The input InAllocatorInstance must be empty; it is only meant to pass parameters to the AllocatorInstance.
+	 * @see TArray(ElementAllocatorType&&) for further comments on the InAllocatorInstance.
+	 */
+	[[nodiscard]] UE_NODEBUG UE_FORCEINLINE_HINT explicit TArray(ElementAllocatorType&& InAllocatorInstance, const ElementType* Ptr, SizeType Count)
+		: AllocatorInstance(MoveTemp(InAllocatorInstance))
+	{
+		check(!AllocatorInstance.HasAllocation());
+		Construct(Ptr, Count);
+	}
+private:
+	UE_FORCEINLINE_HINT void Construct(const ElementType* Ptr, SizeType Count)
 	{
 		if (Count < 0)
 		{
@@ -831,23 +866,36 @@ public:
 
 		CopyToEmpty(Ptr, Count, 0);
 	}
+public:
 
 	template <typename OtherElementType, typename OtherSizeType>
 	[[nodiscard]] explicit TArray(const TArrayView<OtherElementType, OtherSizeType>& Other);
 
+	template <typename OtherElementType, typename OtherSizeType>
+	[[nodiscard]] UE_NODEBUG explicit TArray(ElementAllocatorType&& InAllocatorInstance, const TArrayView<OtherElementType, OtherSizeType>& Other);
+
 	/**
 	 * Initializer list constructor
 	 */
 	[[nodiscard]] TArray(std::initializer_list<InElementType> InitList)
 	{
-		// This is not strictly legal, as std::initializer_list's iterators are not guaranteed to be pointers, but
-		// this appears to be the case on all of our implementations.  Also, if it's not true on a new implementation,
-		// it will fail to compile rather than behave badly.
 		CopyToEmpty(InitList.begin(), (SizeType)InitList.size(), 0);
 	}
 
 	/**
-	 * Copy constructor with changed allocator. Use the common routine to perform the copy.
+	 * Constructor from an initializer list, that also initializes the ElementAllocatorType from a moved-in parameter.
+	 * The input InAllocatorInstance must be empty; it is only meant to pass parameters to the AllocatorInstance.
+	 * @see TArray(ElementAllocatorType&&) for further comments on the InAllocatorInstance.
+	 */
+	[[nodiscard]] UE_NODEBUG explicit TArray(ElementAllocatorType&& InAllocatorInstance, std::initializer_list<InElementType> InitList)
+		: AllocatorInstance(MoveTemp(InAllocatorInstance))
+	{
+		check(!AllocatorInstance.HasAllocation());
+		CopyToEmpty(InitList.begin(), (SizeType)InitList.size(), 0);
+	}
+
+	/**
+	 * Copy constructor from a TArray with a different AllocatorType. Use the common routine to perform the copy.
 	 *
 	 * @param Other The source array to copy.
 	 */
@@ -862,7 +910,25 @@ public:
 	}
 
 	/**
-	 * Copy constructor. Use the common routine to perform the copy.
+	 * Copy Constructor from a TArray with a different AllocatorType, that also initializes this TArray's
+	 * ElementAllocatorType from a moved-in parameter before doing the copy.
+	 * The input InAllocatorInstance must be empty; it is only meant to pass parameters to the AllocatorInstance.
+	 * @see TArray(ElementAllocatorType&&) for further comments on the InAllocatorInstance.
+	 */
+	template <
+		typename OtherElementType,
+		typename OtherAllocator
+		UE_REQUIRES(UE::Core::Private::TArrayElementsAreCompatible_V<ElementType, const OtherElementType&>)
+	>
+	[[nodiscard]] UE_NODEBUG UE_FORCEINLINE_HINT explicit TArray(ElementAllocatorType&& InAllocatorInstance, const TArray<OtherElementType, OtherAllocator>& Other)
+		: AllocatorInstance(MoveTemp(InAllocatorInstance))
+	{
+		check(!AllocatorInstance.HasAllocation());
+		CopyToEmpty(Other.GetData(), Other.Num(), 0);
+	}
+
+	/**
+	 * Copy constructor from a TArray with the same AllocatorType. Use the common routine to perform the copy.
 	 *
 	 * @param Other The source array to copy.
 	 */
@@ -872,7 +938,20 @@ public:
 	}
 
 	/**
-	 * Copy constructor. Use the common routine to perform the copy.
+	 * Copy Constructor from a TArray with the same AllocatorType, that also initializes this TArray's
+	 * ElementAllocatorType from a moved-in parameter before doing the copy.
+	 * The input InAllocatorInstance must be empty; this function is only meant to pass parameters to the AllocatorInstance.
+	 * @see TArray(ElementAllocatorType&&) for further comments on the InAllocatorInstance.
+	 */
+	[[nodiscard]] UE_NODEBUG UE_FORCEINLINE_HINT explicit TArray(ElementAllocatorType&& InAllocatorInstance, const TArray& Other)
+		: AllocatorInstance(MoveTemp(InAllocatorInstance))
+	{
+		check(!AllocatorInstance.HasAllocation());
+		CopyToEmpty(Other.GetData(), Other.Num(), 0);
+	}
+
+	/**
+	 * Copy constructor from TArray with same AllocatorType, with ExtraSlack. Use the common routine to perform the copy.
 	 *
 	 * @param Other The source array to copy.
 	 * @param ExtraSlack Tells how much extra memory should be preallocated
@@ -883,6 +962,19 @@ public:
 		CopyToEmptyWithSlack(Other.GetData(), Other.Num(), 0, ExtraSlack);
 	}
 
+	/**
+	 * Copy constructor from TArray with same AllocatorType, with ExtraSlack, that also initializes this TArray's
+	 * ElementAllocatorType from a moved-in parameter before doing the copy.
+	 * The input InAllocatorInstance must be empty; this function is only meant to pass parameters to the AllocatorInstance.
+	 * @see TArray(ElementAllocatorType&&) for further comments on the InAllocatorInstance.
+	 */
+	[[nodiscard]] UE_NODEBUG UE_FORCEINLINE_HINT explicit TArray(ElementAllocatorType&& InAllocatorInstance, const TArray& Other, SizeType ExtraSlack)
+		: AllocatorInstance(MoveTemp(InAllocatorInstance))
+	{
+		check(!AllocatorInstance.HasAllocation());
+		CopyToEmptyWithSlack(Other.GetData(), Other.Num(), 0, ExtraSlack);
+	}
+
 	/**
 	 * Initializer list assignment operator. First deletes all currently contained elements
 	 * and then copies from initializer list.
@@ -892,9 +984,6 @@ public:
 	TArray& operator=(std::initializer_list<InElementType> InitList)
 	{
 		DestructItems(GetData(), ArrayNum);
-		// This is not strictly legal, as std::initializer_list's iterators are not guaranteed to be pointers, but
-		// this appears to be the case on all of our implementations.  Also, if it's not true on a new implementation,
-		// it will fail to compile rather than behave badly.
 		CopyToEmpty(InitList.begin(), (SizeType)InitList.size(), ArrayMax);
 		return *this;
 	}
@@ -1040,7 +1129,7 @@ private:
 
 public:
 	/**
-	 * Move constructor.
+	 * Move constructor from a TArray with the same AllocatorInstance type.
 	 *
 	 * @param Other Array to move from.
 	 */
@@ -1050,7 +1139,20 @@ public:
 	}
 
 	/**
-	 * Move constructor.
+	 * Move constructor from same AllocatorInstance type that also initializes this TArray's ElementAllocatorType from
+	 * a moved-in parameter before moving the elements (or entire data pointer) from the other TArray's AllocatorInstance.
+	 * The input InAllocatorInstance must be empty; this function is only meant to pass parameters to the AllocatorInstance.
+	 * @see TArray(ElementAllocatorType&&) for further comments on the InAllocatorInstance.
+	 */
+	[[nodiscard]] UE_NODEBUG UE_FORCEINLINE_HINT explicit TArray(ElementAllocatorType&& InAllocatorInstance, TArray&& Other)
+		: AllocatorInstance(MoveTemp(InAllocatorInstance))
+	{
+		check(!AllocatorInstance.HasAllocation());
+		MoveOrCopy(*this, Other, 0);
+	}
+
+	/**
+	 * Move constructor from a TArray with a different AllocatorInstance type.
 	 *
 	 * @param Other Array to move from.
 	 */
@@ -1065,7 +1167,25 @@ public:
 	}
 
 	/**
-	 * Move constructor.
+	 * Move constructor from different AllocatorInstance type that also initializes this TArray's ElementAllocatorType from
+	 * a moved-in parameter before moving the elements from the other TArray's AllocatorInstance.
+	 * The input InAllocatorInstance must be empty; this function is only meant to pass parameters to the AllocatorInstance.
+	 * @see TArray(ElementAllocatorType&&) for further comments on the InAllocatorInstance.
+	 */
+	template <
+		typename OtherElementType,
+		typename OtherAllocator
+		UE_REQUIRES(UE::Core::Private::TArrayElementsAreCompatible_V<ElementType, OtherElementType&&>)
+	>
+	[[nodiscard]] UE_NODEBUG UE_FORCEINLINE_HINT explicit TArray(ElementAllocatorType&& InAllocatorInstance, TArray<OtherElementType, OtherAllocator>&& Other)
+		: AllocatorInstance(MoveTemp(InAllocatorInstance))
+	{
+		check(!AllocatorInstance.HasAllocation());
+		MoveOrCopy(*this, Other, 0);
+	}
+
+	/**
+	 * Move constructor from a TArray with the same or different AllocatorInstance type, with ExtraSlack.
 	 *
 	 * @param Other Array to move from.
 	 * @param ExtraSlack Tells how much extra memory should be preallocated
@@ -1080,6 +1200,24 @@ public:
 		MoveOrCopyWithSlack(*this, Other, 0, ExtraSlack);
 	}
 
+	/**
+	 * Move constructor from same or different AllocatorInstance type, with ExtraSlack, that also initializes this
+	 * TArray's ElementAllocatorType from a moved-in parameter before moving the elements (or entire data pointer) from the other
+	 * TArray's AllocatorInstance. The input InAllocatorInstance must be empty; this function is only meant to pass parameters
+	 * to the AllocatorInstance.
+	 * @see TArray(ElementAllocatorType&&) for further comments on the InAllocatorInstance.
+	 */
+	template <
+		typename OtherElementType
+		UE_REQUIRES(UE::Core::Private::TArrayElementsAreCompatible_V<ElementType, OtherElementType&&>)
+	>
+	[[nodiscard]] UE_NODEBUG explicit TArray(ElementAllocatorType&& InAllocatorInstance, TArray<OtherElementType, AllocatorType>&& Other, SizeType ExtraSlack)
+		: AllocatorInstance(MoveTemp(InAllocatorInstance))
+	{
+		check(!AllocatorInstance.HasAllocation());
+		MoveOrCopyWithSlack(*this, Other, 0, ExtraSlack);
+	}
+
 	/**
 	 * Move assignment operator.
 	 *
diff --git a/Engine/Source/Runtime/Core/Public/Containers/ArrayView.h b/Engine/Source/Runtime/Core/Public/Containers/ArrayView.h
--- a/Engine/Source/Runtime/Core/Public/Containers/ArrayView.h
+++ b/Engine/Source/Runtime/Core/Public/Containers/ArrayView.h
@@ -1063,6 +1063,15 @@ template<typename OtherElementType, typename OtherSizeType>
 	CopyToEmpty(Other.GetData(), Other.Num(), 0);
 }
 
+template<typename InElementType, typename InAllocatorType>
+template<typename OtherElementType, typename OtherSizeType>
+[[nodiscard]] UE_FORCEINLINE_HINT TArray<InElementType, InAllocatorType>::TArray(ElementAllocatorType&& InAllocatorInstance, const TArrayView<OtherElementType, OtherSizeType>& Other)
+	: AllocatorInstance(MoveTemp(InAllocatorInstance))
+{
+	check(!AllocatorInstance.HasAllocation());
+	CopyToEmpty(Other.GetData(), Other.Num(), 0);
+}
+
 template<typename InElementType, typename InAllocatorType>
 template<typename OtherElementType, typename OtherSizeType>
 inline TArray<InElementType, InAllocatorType>& TArray<InElementType, InAllocatorType>::operator=(const TArrayView<OtherElementType, OtherSizeType>& Other)
diff --git a/Engine/Source/Runtime/Core/Public/Containers/ContainerAllocationPolicies.h b/Engine/Source/Runtime/Core/Public/Containers/ContainerAllocationPolicies.h
--- a/Engine/Source/Runtime/Core/Public/Containers/ContainerAllocationPolicies.h
+++ b/Engine/Source/Runtime/Core/Public/Containers/ContainerAllocationPolicies.h
@@ -8,6 +8,7 @@
 #include "HAL/UnrealMemory.h"
 #include "Math/NumericLimits.h"
 #include "Misc/AssertionMacros.h"
+#include "Misc/CoreMiscDefines.h"
 #include "Templates/AlignmentTemplates.h"
 #include "Templates/IsPolymorphic.h"
 #include "Templates/MemoryOps.h"
@@ -441,7 +442,7 @@ namespace UE::Core::Private
 	[[noreturn]] CORE_API void OnInvalidSizedHeapAllocatorNum(int32 IndexSize, int64 NewNum, SIZE_T NumBytesPerElement);
 }
 
-/** The indirect allocation policy always allocates the elements indirectly. */
+/** A TArray allocator (32-bit sizetype only) that allocates from FMemory and allows specifying alignment for all allocations. */
 template<uint32 Alignment = DEFAULT_ALIGNMENT>
 class TAlignedHeapAllocator
 {
@@ -612,15 +613,19 @@ template <> struct TBitsToSizeType<16> { using Type = int16; };
 template <> struct TBitsToSizeType<32> { using Type = int32; };
 template <> struct TBitsToSizeType<64> { using Type = int64; };
 
-/** The indirect allocation policy always allocates the elements indirectly. */
-template <int IndexSize, typename BaseMallocType = FMemory>
-class TSizedHeapAllocator
+/**
+ * BaseClass for TArray Allocators.
+ * Templated to support the container's sizetype (e.g. int32 (IndexSize==32) or int64 (IndexSize==64)).
+ * Uses non-virtual inheritance via the CuriouslyRecurringTemplatePattern; concrete subclasses must pass
+ * their own type in as the SubClass template parameter.
+ */
+template <int IndexSize, typename SubClass>
+class TSizedAllocatorBase
 {
 public:
 	using SizeType = typename TBitsToSizeType<IndexSize>::Type;
-	using BaseMalloc = BaseMallocType;
 
-private:
+protected:
 	using USizeType = std::make_unsigned_t<SizeType>;
 
 public:
@@ -630,16 +635,17 @@ public:
 	class ForAnyElementType
 	{
 		template <int, typename>
-		friend class TSizedHeapAllocator;
+		friend class TSizedAllocatorBase;
 
 	public:
 		/** Default constructor. */
-		constexpr ForAnyElementType()
+		UE_NODEBUG constexpr ForAnyElementType()
 			: Data(nullptr)
-		{}
+		{
+		}
 
 		/** Explicitly consteval constructor. */
-		explicit consteval ForAnyElementType(EConstEval)
+		UE_NODEBUG explicit consteval ForAnyElementType(EConstEval)
 			: Data(nullptr)
 		{}
 
@@ -650,17 +656,14 @@ public:
 		 * @param Other - The allocator to move the state from.  This allocator should be left in a valid empty state.
 		 */
 		template <typename OtherAllocator>
-		UE_FORCEINLINE_HINT void MoveToEmptyFromOtherAllocator(typename OtherAllocator::ForAnyElementType& Other)
+		UE_NODEBUG UE_FORCEINLINE_HINT void MoveToEmptyFromOtherAllocator(typename OtherAllocator::ForAnyElementType& Other)
 		{
 			checkSlow((void*)this != (void*)&Other);
 
 			if (Data)
 			{
-#if UE_ENABLE_ARRAY_SLACK_TRACKING
-				FArraySlackTrackingHeader::Free(Data);
-#else
-				BaseMallocType::Free(Data);
-#endif
+				// Call subclass Free using CuriouslyRecurringTemplatePattern.
+				static_cast<SubClass::ForAnyElementType*>(this)->CallFree(Data);
 			}
 
 			Data = Other.Data;
@@ -668,36 +671,32 @@ public:
 		}
 
 		/**
-		 * Moves the state of another allocator into this one.
-		 * Moves the state of another allocator into this one.  The allocator can be different.
+		 * Moves the state of another allocator into this one. The allocator is the same class.
 		 *
 		 * Assumes that the allocator is currently empty, i.e. memory may be allocated but any existing elements have already been destructed (if necessary).
 		 * @param Other - The allocator to move the state from.  This allocator should be left in a valid empty state.
 		 */
-		UE_FORCEINLINE_HINT void MoveToEmpty(ForAnyElementType& Other)
+		UE_NODEBUG UE_FORCEINLINE_HINT void MoveToEmpty(SubClass::ForAnyElementType& Other)
 		{
-			this->MoveToEmptyFromOtherAllocator<TSizedHeapAllocator>(Other);
+			this->MoveToEmptyFromOtherAllocator<SubClass>(Other);
 		}
 
 		/** Destructor. */
-		UE_FORCEINLINE_HINT ~ForAnyElementType()
+		UE_NODEBUG UE_FORCEINLINE_HINT ~ForAnyElementType()
 		{
-			if(Data)
+			if (Data)
 			{
-#if UE_ENABLE_ARRAY_SLACK_TRACKING
-				FArraySlackTrackingHeader::Free(Data);
-#else
-				BaseMallocType::Free(Data);
-#endif
+				// Call subclass Free using CuriouslyRecurringTemplatePattern.
+				static_cast<SubClass::ForAnyElementType*>(this)->CallFree(Data);
 			}
 		}
 
 		// FContainerAllocatorInterface
-		UE_FORCEINLINE_HINT FScriptContainerElement* GetAllocation() const
+		UE_NODEBUG UE_FORCEINLINE_HINT FScriptContainerElement* GetAllocation() const
 		{
 			return Data;
 		}
-		void ResizeAllocation(SizeType CurrentNum, SizeType NewMax, SIZE_T NumBytesPerElement)
+		UE_NODEBUG void ResizeAllocation(SizeType CurrentNum, SizeType NewMax, SIZE_T NumBytesPerElement)
 		{
 			// Avoid calling FMemory::Realloc( nullptr, 0 ) as ANSI C mandates returning a valid pointer which is not what we want.
 			if (Data || NewMax)
@@ -715,16 +714,13 @@ public:
 					UE::Core::Private::OnInvalidSizedHeapAllocatorNum(IndexSize, NewMax, NumBytesPerElement);
 				}
 
-#if UE_ENABLE_ARRAY_SLACK_TRACKING
-				Data = (FScriptContainerElement*)FArraySlackTrackingHeader::Realloc(Data, NewMax, NumBytesPerElement, 0);
-#else
-				Data = (FScriptContainerElement*)BaseMallocType::Realloc( Data, NewMax*NumBytesPerElement );
-#endif
+				// Call subclass Realloc using CuriouslyRecurringTemplatePattern.
+				Data = (FScriptContainerElement*)static_cast<SubClass::ForAnyElementType*>(this)->CallRealloc(Data, NewMax, NumBytesPerElement);
 			}
 		}
-		void ResizeAllocation(SizeType CurrentNum, SizeType NewMax, SIZE_T NumBytesPerElement, uint32 AlignmentOfElement)
+		UE_NODEBUG void ResizeAllocation(SizeType CurrentNum, SizeType NewMax, SIZE_T NumBytesPerElement, uint32 AlignmentOfElement)
 		{
-			// Avoid calling FMemory::Realloc( nullptr, 0 ) as ANSI C mandates returning a valid pointer which is not what we want.
+			// Avoid calling Realloc( nullptr, 0 ) as ANSI C mandates returning a valid pointer which is not what we want.
 			if (Data || NewMax)
 			{
 				static_assert(sizeof(SizeType) <= sizeof(SIZE_T), "SIZE_T is expected to handle all possible sizes");
@@ -740,68 +736,73 @@ public:
 					UE::Core::Private::OnInvalidSizedHeapAllocatorNum(IndexSize, NewMax, NumBytesPerElement);
 				}
 
-#if UE_ENABLE_ARRAY_SLACK_TRACKING
-				Data = (FScriptContainerElement*)FArraySlackTrackingHeader::Realloc(Data, NewMax, NumBytesPerElement, AlignmentOfElement > alignof(FArraySlackTrackingHeader) ? AlignmentOfElement : alignof(FArraySlackTrackingHeader));
-#else
-				Data = (FScriptContainerElement*)BaseMallocType::Realloc( Data, NewMax*NumBytesPerElement, AlignmentOfElement );
-#endif
+				// Call subclass Realloc using CuriouslyRecurringTemplatePattern.
+				Data = (FScriptContainerElement*)static_cast<SubClass::ForAnyElementType*>(this)->CallRealloc(Data, NewMax, NumBytesPerElement, AlignmentOfElement);
 			}
 		}
-		UE_FORCEINLINE_HINT SizeType CalculateSlackReserve(SizeType NewMax, SIZE_T NumBytesPerElement) const
+		UE_NODEBUG UE_FORCEINLINE_HINT SizeType CalculateSlackReserve(SizeType NewMax, SIZE_T NumBytesPerElement) const
 		{
 			return DefaultCalculateSlackReserve(NewMax, NumBytesPerElement, true);
 		}
-		UE_FORCEINLINE_HINT SizeType CalculateSlackReserve(SizeType NewMax, SIZE_T NumBytesPerElement, uint32 AlignmentOfElement) const
+		UE_NODEBUG UE_FORCEINLINE_HINT SizeType CalculateSlackReserve(SizeType NewMax, SIZE_T NumBytesPerElement, uint32 AlignmentOfElement) const
 		{
 			return DefaultCalculateSlackReserve(NewMax, NumBytesPerElement, true, (uint32)AlignmentOfElement);
 		}
-		UE_FORCEINLINE_HINT SizeType CalculateSlackShrink(SizeType NewMax, SizeType CurrentMax, SIZE_T NumBytesPerElement) const
+		UE_NODEBUG UE_FORCEINLINE_HINT SizeType CalculateSlackShrink(SizeType NewMax, SizeType CurrentMax, SIZE_T NumBytesPerElement) const
 		{
 			return DefaultCalculateSlackShrink(NewMax, CurrentMax, NumBytesPerElement, true);
 		}
-		UE_FORCEINLINE_HINT SizeType CalculateSlackShrink(SizeType NewMax, SizeType CurrentMax, SIZE_T NumBytesPerElement, uint32 AlignmentOfElement) const
+		UE_NODEBUG UE_FORCEINLINE_HINT SizeType CalculateSlackShrink(SizeType NewMax, SizeType CurrentMax, SIZE_T NumBytesPerElement, uint32 AlignmentOfElement) const
 		{
 			return DefaultCalculateSlackShrink(NewMax, CurrentMax, NumBytesPerElement, true, (uint32)AlignmentOfElement);
 		}
-		UE_FORCEINLINE_HINT SizeType CalculateSlackGrow(SizeType NewMax, SizeType CurrentMax, SIZE_T NumBytesPerElement) const
+		UE_NODEBUG UE_FORCEINLINE_HINT SizeType CalculateSlackGrow(SizeType NewMax, SizeType CurrentMax, SIZE_T NumBytesPerElement) const
 		{
 			return DefaultCalculateSlackGrow(NewMax, CurrentMax, NumBytesPerElement, true);
 		}
-		UE_FORCEINLINE_HINT SizeType CalculateSlackGrow(SizeType NewMax, SizeType CurrentMax, SIZE_T NumBytesPerElement, uint32 AlignmentOfElement) const
+		UE_NODEBUG UE_FORCEINLINE_HINT SizeType CalculateSlackGrow(SizeType NewMax, SizeType CurrentMax, SIZE_T NumBytesPerElement, uint32 AlignmentOfElement) const
 		{
 			return DefaultCalculateSlackGrow(NewMax, CurrentMax, NumBytesPerElement, true, (uint32)AlignmentOfElement);
 		}
 
-		SIZE_T GetAllocatedSize(SizeType CurrentMax, SIZE_T NumBytesPerElement) const
+		UE_NODEBUG SIZE_T GetAllocatedSize(SizeType CurrentMax, SIZE_T NumBytesPerElement) const
 		{
 			return CurrentMax * NumBytesPerElement;
 		}
 
-		bool HasAllocation() const
+		UE_NODEBUG bool HasAllocation() const
 		{
 			return !!Data;
 		}
 
-		constexpr SizeType GetInitialCapacity() const
+		UE_NODEBUG constexpr SizeType GetInitialCapacity() const
 		{
 			return 0;
 		}
 
+	public:
+		// Optional template functions to override on ForAnyElementType on the SubClass.
 #if UE_ENABLE_ARRAY_SLACK_TRACKING
-		UE_FORCEINLINE_HINT void SlackTrackerLogNum(SizeType NewNumUsed)
+		UE_NODEBUG UE_FORCEINLINE_HINT void SlackTrackerLogNum(SizeType NewNumUsed)
 		{
-			FArraySlackTrackingHeader::UpdateNumUsed(Data, (int64)NewNumUsed);
+			// Base implementation does not report to SlackTracking; subclasses must override if they need to report
+			// to slacktracking, @see TSizedHeapAllocator. 
 		}
 
 		// Suppress slack tracking on an allocation -- should be called whenever the container may have been resized.
 		// Useful for debug allocations you don't want to show up in slack reports.
-		UE_FORCEINLINE_HINT void DisableSlackTracking()
+		UE_NODEBUG UE_FORCEINLINE_HINT void DisableSlackTracking()
 		{
-			FArraySlackTrackingHeader::DisableTracking(Data);
+			// Base implementation does not report to SlackTracking; subclasses must override if they need to report
+			// to slacktracking, @see TSizedHeapAllocator. 
 		}
 #endif
 
-	private:
+		// Required template functions to override on ForAnyElementType on the SubClass.
+		void* CallRealloc(void* InData, SizeType NewMax, SIZE_T NumBytesPerElement, uint32 AlignmentOfElement = DEFAULT_ALIGNMENT) = delete;
+		void CallFree(void* InData) = delete;
+
+	protected:
 		ForAnyElementType(const ForAnyElementType&);
 		ForAnyElementType& operator=(const ForAnyElementType&);
 
@@ -810,15 +811,199 @@ public:
 	};
 
 	template<typename ElementType>
-	class ForElementType : public ForAnyElementType
+	class ForElementType : public SubClass::ForAnyElementType
 	{
 	public:
-		/** Default constructor. */
 		ForElementType() = default;
-		explicit consteval ForElementType(EConstEval)
+		UE_NODEBUG UE_FORCEINLINE_HINT ElementType* GetAllocation() const
+		{
+			return (ElementType*)ForAnyElementType::GetAllocation();
+		}
+	};
+};
+
+/**
+ * A TArray allocator that allocates from FMemory.
+ * The BaseMallocType template parameter is deprecated and passing it in will be a compile error in a future release.
+ * Use TSizedExternalAllocator instead to provide an ExternalMallocType,
+ */
+template <int IndexSize, typename BaseMallocType = FMemory>
+class TSizedHeapAllocator : public TSizedAllocatorBase<IndexSize, TSizedHeapAllocator<IndexSize, BaseMallocType>>
+{
+public:
+	using MyClass = TSizedHeapAllocator<IndexSize, BaseMallocType>;
+	using SuperClass = TSizedAllocatorBase<IndexSize, MyClass>;
+	using BaseMalloc = BaseMallocType;
+	using SizeType = SuperClass::SizeType;
+
+	class ForAnyElementType : public SuperClass::ForAnyElementType
+	{
+	public:
+		// Declare we are using constructors from ForAnyElementType on SuperClass
+		using SuperClass::ForAnyElementType::ForAnyElementType;
+
+		// Declare we are using Data pointer on the superclass; the compiler can't figure this out automatically due
+		// to two-phase template compile.
+		using SuperClass::ForAnyElementType::Data;
+
+		// Declare our constructor so we can add deprecation handling for BaseMallocType.
+		UE_NODEBUG constexpr ForAnyElementType()
+		{
+			UE_STATIC_DEPRECATE(5.8, (!std::is_same_v<BaseMallocType, FMemory>), "Overriding the MallocType in TSizedHeapAllocator is deprecated. Use TSizedExternalAllocator<int IndexSize, typename ExternalMallocType> instead");
+		}
+
+		// Override the ResizeAllocation functions so they can be defined as exported to avoid bloat
+#if !UE_MERGED_MODULES && !defined(UE_HEADER_UNITS)
+		UE_NODEBUG void ResizeAllocation(SizeType CurrentNum, SizeType NewMax, SIZE_T NumBytesPerElement)
+		{
+			SuperClass::ForAnyElementType::ResizeAllocation(CurrentNum, NewMax, NumBytesPerElement);
+		}
+		UE_NODEBUG void ResizeAllocation(SizeType CurrentNum, SizeType NewMax, SIZE_T NumBytesPerElement, uint32 AlignmentOfElement)
+		{
+			SuperClass::ForAnyElementType::ResizeAllocation(CurrentNum, NewMax, NumBytesPerElement, AlignmentOfElement);
+		}
+#endif
+
+	public:
+#if UE_ENABLE_ARRAY_SLACK_TRACKING
+		UE_NODEBUG UE_FORCEINLINE_HINT void SlackTrackerLogNum(SizeType NewNumUsed)
+		{
+			FArraySlackTrackingHeader::UpdateNumUsed(Data, (int64)NewNumUsed);
+		}
+
+		UE_NODEBUG UE_FORCEINLINE_HINT void DisableSlackTracking()
+		{
+			FArraySlackTrackingHeader::DisableTracking(Data);
+		}
+#endif
+
+		UE_NODEBUG UE_FORCEINLINE_HINT void* CallRealloc(void* InData, SuperClass::SizeType NewMax,
+			SIZE_T NumBytesPerElement, uint32 AlignmentOfElement = DEFAULT_ALIGNMENT)
+		{
+#if UE_ENABLE_ARRAY_SLACK_TRACKING
+			AlignmentOfElement = AlignmentOfElement == 0
+				? 0
+				: (AlignmentOfElement > alignof(FArraySlackTrackingHeader)
+					? AlignmentOfElement : alignof(FArraySlackTrackingHeader));
+			return (FScriptContainerElement*)FArraySlackTrackingHeader::Realloc(InData, NewMax, NumBytesPerElement, AlignmentOfElement);
+#else
+			return (FScriptContainerElement*)BaseMalloc::Realloc(InData, NewMax * NumBytesPerElement, AlignmentOfElement);
+#endif
+		}
+		UE_NODEBUG UE_FORCEINLINE_HINT void CallFree(void* InData)
+		{
+#if UE_ENABLE_ARRAY_SLACK_TRACKING
+			FArraySlackTrackingHeader::Free(InData);
+#else
+			BaseMalloc::Free(InData);
+#endif
+		}
+	};
+
+	// Special implementation of ForElementType that implements the consteval constructor.
+	template<typename ElementType>
+	class ForElementType : public ForAnyElementType
+	{
+	public:
+		UE_NODEBUG ForElementType() = default;
+		UE_NODEBUG explicit consteval ForElementType(EConstEval)
 			: ForAnyElementType(ConstEval)
 		{
 		}
+		UE_NODEBUG UE_FORCEINLINE_HINT ElementType* GetAllocation() const
+		{
+			return (ElementType*)ForAnyElementType::GetAllocation();
+		}
+	};
+};
+
+/** A SizedAllocator that allocates from a compile-time-provided set of allocation functions. */
+template <int IndexSize, typename ExternalMallocType>
+class TSizedExternalAllocator : public TSizedAllocatorBase<IndexSize, TSizedExternalAllocator<IndexSize, ExternalMallocType>>
+{
+public:
+	using MyClass = TSizedExternalAllocator<IndexSize, ExternalMallocType>;
+	using SuperClass = TSizedAllocatorBase<IndexSize, MyClass>;
+	using ExternalMalloc = ExternalMallocType;
+
+	class ForAnyElementType : public SuperClass::ForAnyElementType
+	{
+	public:
+		// Declare we are using constructors from ForAnyElementType on SuperClass
+		using SuperClass::ForAnyElementType::ForAnyElementType;
+
+	public:
+		UE_FORCEINLINE_HINT void* CallRealloc(void* InData, SuperClass::SizeType NewMax,
+			SIZE_T NumBytesPerElement, uint32 AlignmentOfElement = DEFAULT_ALIGNMENT)
+		{
+			return (FScriptContainerElement*)ExternalMalloc::Realloc(InData, NewMax * NumBytesPerElement, AlignmentOfElement);
+		}
+		UE_FORCEINLINE_HINT void CallFree(void* InData)
+		{
+			ExternalMalloc::Free(InData);
+		}
+	};
+};
+
+/**
+ * A SizedAllocator that allocates from a runtime-provided pointer to an FMalloc interface.
+ * TArrays using this allocator must be constructed with the TArray constructors that take an ElementAllocatorType&&
+ * argument; attempting to construct them without that argument will intentionally fail due to lack of a default constructor
+ * on ForAnyElementType.
+*/
+template <int IndexSize>
+class TSizedMallocAllocator : public TSizedAllocatorBase<IndexSize, TSizedMallocAllocator<IndexSize>>
+{
+public:
+	using MyClass = TSizedMallocAllocator<IndexSize>;
+	using SuperClass = TSizedAllocatorBase<IndexSize, MyClass>;
+
+	class ForAnyElementType : public SuperClass::ForAnyElementType
+	{
+	public:
+		// Declare we are using Data pointer on the superclass; the compiler can't figure this out automatically due
+		// to two-phase template compile.
+		using SuperClass::ForAnyElementType::Data;
+
+		explicit ForAnyElementType(FMalloc* InMalloc)
+			: Malloc(InMalloc)
+		{
+		}
+		ForAnyElementType(ForAnyElementType&& Other)
+			: Malloc(Other.Malloc)
+		{
+			Data = Other.Data;
+			Other.Data = nullptr;
+		}
+
+	public:
+		UE_FORCEINLINE_HINT void* CallRealloc(void* InData, SuperClass::SizeType NewMax,
+			SIZE_T NumBytesPerElement, uint32 AlignmentOfElement = DEFAULT_ALIGNMENT)
+		{
+			return Malloc->Realloc(InData, NewMax * NumBytesPerElement, AlignmentOfElement);
+		}
+		UE_FORCEINLINE_HINT void CallFree(void* InData)
+		{
+			Malloc->Free(InData);
+		}
+
+	protected:
+		FMalloc* Malloc;
+	};
+
+	// Declare our own ForElementType so we can support our constructor parameters.
+	template<typename ElementType>
+	class ForElementType : public ForAnyElementType
+	{
+	public:
+		explicit ForElementType(FMalloc* InMalloc)
+			: ForAnyElementType(InMalloc)
+		{
+		}
+		ForElementType(ForElementType&& Other)
+			: ForAnyElementType(MoveTemp(Other))
+		{
+		}
 
 		UE_FORCEINLINE_HINT ElementType* GetAllocation() const
 		{
@@ -830,6 +1015,8 @@ public:
 template <int IndexSize> class TSizedDefaultAllocator : public TSizedHeapAllocator<IndexSize> { public: typedef TSizedHeapAllocator<IndexSize> Typedef; };
 using FDefaultAllocator = TSizedDefaultAllocator<32>;
 
+using FDefaultMallocAllocator = TSizedMallocAllocator<32>;
+
 /** TSizedNonshrinkingAllocator behaves like TSizedDefaultAllocator, but prevents arrays from automatically shrinking unless EAllowShrinking::Yes is used. */
 template <int IndexSize> class TSizedNonshrinkingAllocator : public TSizedHeapAllocator<IndexSize>
 {
