Color: Minor fixes to HLG OOTF implementation from previous CL 52237635.

* Avoid operating on alpha, avoid NaN for zero luminance pixels.

#Synced-CL 52172240

diff --git a/Engine/Shaders/Private/GammaCorrectionCommon.ush b/Engine/Shaders/Private/GammaCorrectionCommon.ush
--- a/Engine/Shaders/Private/GammaCorrectionCommon.ush
+++ b/Engine/Shaders/Private/GammaCorrectionCommon.ush
@@ -269,6 +269,6 @@ float3 HLGToLinearWithOOTF(float3 Color, float PeakLuminance = 1000.0f, float Sy
 	Color = HLGToLinear(Color);
 	float Ys = dot(Color, Rec2020LuminanceWeights);
 
-	return PeakLuminance * pow(Ys, SystemGamma - 1.0f) * Color;
+	return select(Ys > 0.0f, PeakLuminance * pow(Ys, SystemGamma - 1.0f) * Color, float3(0.0f, 0.0f, 0.0f));
 }
 
diff --git a/Engine/Source/Runtime/Core/Private/ColorManagement/Tests/ColorManagementTests.cpp b/Engine/Source/Runtime/Core/Private/ColorManagement/Tests/ColorManagementTests.cpp
--- a/Engine/Source/Runtime/Core/Private/ColorManagement/Tests/ColorManagementTests.cpp
+++ b/Engine/Source/Runtime/Core/Private/ColorManagement/Tests/ColorManagementTests.cpp
@@ -353,6 +353,32 @@ TEST_CASE_NAMED(FHLG2100TransferFunctionsTest, "System::ColorManagement::HLG2100
 		CHECK_MESSAGE(TEXT("BT.2100 round-trip with black luminance should preserve color values"), bRClose && bGClose && bBClose);
 	}
 
+	SECTION("BT.2100 black input does not produce NaN at low peak luminance")
+	{
+		// PeakLuminance below ~302 cd/m^2 yields Gamma < 1 via HLG_SystemGamma_Extended,
+		// making the OOTF exponent (Gamma-1) negative. Without a zero-luminance guard,
+		// pow(0, negative) = +inf and inf * 0 = NaN.
+		const FLinearColor Black(0.0f, 0.0f, 0.0f, 1.0f);
+		const float LowPeakLuminances[] = { 100.0f, 200.0f, 300.0f };
+
+		for (float Peak : LowPeakLuminances)
+		{
+			FLinearColor Encoded = EncodeHLG_2100(Black, Peak);
+			CHECK_MESSAGE(TEXT("Encoding black at low peak luminance must not produce NaN"),
+				!FMath::IsNaN(Encoded.R) && !FMath::IsNaN(Encoded.G) && !FMath::IsNaN(Encoded.B));
+
+			FLinearColor Decoded = DecodeHLG_2100(Encoded, Peak);
+			CHECK_MESSAGE(TEXT("Decoding black at low peak luminance must not produce NaN"),
+				!FMath::IsNaN(Decoded.R) && !FMath::IsNaN(Decoded.G) && !FMath::IsNaN(Decoded.B));
+
+			CHECK_MESSAGE(TEXT("Black should round-trip to black RGB"),
+				FMath::IsNearlyZero(Decoded.R) && FMath::IsNearlyZero(Decoded.G) && FMath::IsNearlyZero(Decoded.B));
+
+			CHECK_MESSAGE(TEXT("Alpha should be preserved through BT.2100 round-trip"),
+				FMath::IsNearlyEqual(Decoded.A, 1.0f, UE_SMALL_NUMBER));
+		}
+	}
+
 }
 
 TEST_CASE_NAMED(FColorSpaceTest, "System::ColorManagement::ColorSpace", "[EditorContext][EngineFilter]")
diff --git a/Engine/Source/Runtime/Core/Private/ColorManagement/TransferFunctions.cpp b/Engine/Source/Runtime/Core/Private/ColorManagement/TransferFunctions.cpp
--- a/Engine/Source/Runtime/Core/Private/ColorManagement/TransferFunctions.cpp
+++ b/Engine/Source/Runtime/Core/Private/ColorManagement/TransferFunctions.cpp
@@ -62,7 +62,13 @@ FLinearColor HLG_OOTF(const FLinearColor& Value, const float PeakLuminance, cons
 	const float Alpha = PeakLuminance;
 	const float Ys = FColorSpace::GetRec2020().GetLuminance(Value);
 
-	return Alpha * FMath::Pow(Ys, Gamma - 1.0f) * Value;
+	if (Ys <= 0.0f)
+	{
+		return FLinearColor(0.0f, 0.0f, 0.0f, Value.A);
+	}
+
+	const float Scale = Alpha * FMath::Pow(Ys, Gamma - 1.0f);
+	return FLinearColor(Value.R * Scale, Value.G * Scale, Value.B * Scale, Value.A);
 }
 
 FLinearColor HLG_InverseOOTF(const FLinearColor& Value, const float PeakLuminance, const float Gamma)
@@ -72,10 +78,11 @@ FLinearColor HLG_InverseOOTF(const FLinearColor& Value, const float PeakLuminanc
 
 	if (Yd <= 0.0f)
 	{
-		return FLinearColor::Black;
+		return FLinearColor(0.0f, 0.0f, 0.0f, Value.A);
 	}
 
-	return FMath::Pow(Yd / Alpha, (1.0f - Gamma) / Gamma) * Value / Alpha;
+	const float Scale = FMath::Pow(Yd / Alpha, (1.0f - Gamma) / Gamma) / Alpha;
+	return FLinearColor(Value.R * Scale, Value.G * Scale, Value.B * Scale, Value.A);
 }
 
 } // anonymous namespace
@@ -107,7 +114,7 @@ FLinearColor EncodeHLG_2100(const FLinearColor& LinearDisplayReferred, const flo
 		Beta = 1.0f - UE_SMALL_NUMBER;
 	}
 
-	// The inverse OOTF maps relative display to scene
+	// The inverse OOTF maps relative display to scene (preserves alpha)
 	FLinearColor Output = HLG_InverseOOTF(LinearDisplayReferred, PeakLuminance, Gamma);
 
 	// Encode into HLG signal using OETF
@@ -115,8 +122,13 @@ FLinearColor EncodeHLG_2100(const FLinearColor& LinearDisplayReferred, const flo
 	Output.G = EncodeHLG(Output.G);
 	Output.B = EncodeHLG(Output.B);
 
-	// Invert user black lift
-	return (Output - FLinearColor::White * Beta) / (1.0f - Beta);
+	// Invert user black lift (RGB only, preserve alpha)
+	const float InvBetaRange = 1.0f / (1.0f - Beta);
+	Output.R = (Output.R - Beta) * InvBetaRange;
+	Output.G = (Output.G - Beta) * InvBetaRange;
+	Output.B = (Output.B - Beta) * InvBetaRange;
+
+	return Output;
 }
 
 FLinearColor DecodeHLG_2100(const FLinearColor& EncodedSceneReferred, const float PeakLuminance, const float BlackLuminance)
@@ -147,10 +159,12 @@ FLinearColor DecodeHLG_2100(const FLinearColor& EncodedSceneReferred, const floa
 	}
 
 	// Apply user black lift, clamping to zero per BT.2100-3 EOTF: max(0, (1-beta)*E' + beta)
-	FLinearColor Output = (1.0f - Beta) * EncodedSceneReferred + FLinearColor::White * Beta;
-	Output.R = FMath::Max(0.0f, Output.R);
-	Output.G = FMath::Max(0.0f, Output.G);
-	Output.B = FMath::Max(0.0f, Output.B);
+	// Applied to RGB only to preserve alpha.
+	const float BetaRange = 1.0f - Beta;
+	FLinearColor Output = EncodedSceneReferred;
+	Output.R = FMath::Max(0.0f, BetaRange * Output.R + Beta);
+	Output.G = FMath::Max(0.0f, BetaRange * Output.G + Beta);
+	Output.B = FMath::Max(0.0f, BetaRange * Output.B + Beta);
 
 	// Decode HLG signal using inverse OETF
 	Output.R = DecodeHLG(Output.R);
