Added runtime network status monitoring to platforms missing it (iOS, Android, Mac, Linux)
Windows platform delayed until performance issues are investigated

#Synced-CL 52250008

#Synced-CL 52251831

diff --git a/Engine/Source/Runtime/ApplicationCore/Private/Linux/LinuxApplication.cpp b/Engine/Source/Runtime/ApplicationCore/Private/Linux/LinuxApplication.cpp
--- a/Engine/Source/Runtime/ApplicationCore/Private/Linux/LinuxApplication.cpp
+++ b/Engine/Source/Runtime/ApplicationCore/Private/Linux/LinuxApplication.cpp
@@ -222,6 +222,9 @@ TSharedPtr< FLinuxWindow > FLinuxApplication::FindWindowBySDLWindow(SDL_Window *
 void FLinuxApplication::PumpMessages( const float TimeDelta )
 {
 	FPlatformApplicationMisc::PumpMessages( true );
+
+	// Check for netlink network change events and update status if needed.
+	FPlatformMisc::CheckNetworkConnectionEvents();
 }
 
 bool FLinuxApplication::IsCursorDirectlyOverSlateWindow() const
diff --git a/Engine/Source/Runtime/Core/Private/Android/AndroidPlatformMisc.cpp b/Engine/Source/Runtime/Core/Private/Android/AndroidPlatformMisc.cpp
--- a/Engine/Source/Runtime/Core/Private/Android/AndroidPlatformMisc.cpp
+++ b/Engine/Source/Runtime/Core/Private/Android/AndroidPlatformMisc.cpp
@@ -603,6 +603,9 @@ void FAndroidMisc::PlatformInit()
 
 	extern void AndroidThunkJava_AddNetworkListener();
 	AndroidThunkJava_AddNetworkListener();
+
+	// Set initial network connection status
+	FPlatformMisc::SetNetworkConnectionStatus(FPlatformMisc::NetworkConnectionTypeToStatus(GetNetworkConnectionType()));
 #endif
 
 	InitCpuThermalSensor();
@@ -3124,9 +3127,16 @@ void JNICALL UE::Jni::FGameActivity::nativeNetworkChanged(JNIEnv* jenv, jobject
 	{
 		FFunctionGraphTask::CreateAndDispatchWhenReady([]()
 		{
-			FCoreDelegates::OnNetworkConnectionChanged.Broadcast(FAndroidMisc::GetNetworkConnectionType());
+			const ENetworkConnectionType ConnectionType = FAndroidMisc::GetNetworkConnectionType();
+			FCoreDelegates::OnNetworkConnectionChanged.Broadcast(ConnectionType);
+			FPlatformMisc::SetNetworkConnectionStatus(FPlatformMisc::NetworkConnectionTypeToStatus(ConnectionType));
 		}, TStatId(), NULL, ENamedThreads::GameThread);
 	}
+	else
+	{
+		// Task graph not running yet (early startup). Set the status synchronously so GetNetworkConnectionStatus() stays in sync with GetNetworkConnectionType().
+		FPlatformMisc::SetNetworkConnectionStatus(FPlatformMisc::NetworkConnectionTypeToStatus(LastNetworkConnectionType));
+	}
 }
 
 #endif
diff --git a/Engine/Source/Runtime/Core/Private/Apple/ApplePlatformMisc.cpp b/Engine/Source/Runtime/Core/Private/Apple/ApplePlatformMisc.cpp
--- a/Engine/Source/Runtime/Core/Private/Apple/ApplePlatformMisc.cpp
+++ b/Engine/Source/Runtime/Core/Private/Apple/ApplePlatformMisc.cpp
@@ -544,46 +544,80 @@ void FApplePlatformMisc::PlatformInit()
 {
 	// start monitoring network status
 	UE::FApplePlatformMisc::PathMonitor = nw_path_monitor_create();
-	nw_path_monitor_set_update_handler(
-		UE::FApplePlatformMisc::PathMonitor, 
-		^(nw_path_t InPath)
-		{
-			bool bIsConstrained = nw_path_is_constrained(InPath);
-			bool bIsExpensive   = nw_path_is_expensive(InPath);
-
-			UE_LOG(LogInit, Log, TEXT("Network path changed (IsConstrained=%d, IsExpensive=%d)"), bIsConstrained, bIsExpensive);
-			
-			// cache the path we get
-			if (UE::FApplePlatformMisc::CachedNetworkPath != nullptr)
-			{
-				nw_release(UE::FApplePlatformMisc::CachedNetworkPath);
-			}
-			UE::FApplePlatformMisc::CachedNetworkPath = InPath;
-			nw_retain(UE::FApplePlatformMisc::CachedNetworkPath);
-			if (FTaskGraphInterface::IsRunning())
+	if (UE::FApplePlatformMisc::PathMonitor == nullptr)
+	{
+		UE_LOG(LogApplePlatformMisc, Warning, TEXT("Failed to create nw_path_monitor for network monitoring"));
+	}
+	else
+	{
+		nw_path_monitor_set_update_handler(
+			UE::FApplePlatformMisc::PathMonitor,
+			^(nw_path_t InPath)
 			{
+				bool bIsConstrained = nw_path_is_constrained(InPath);
+				bool bIsExpensive   = nw_path_is_expensive(InPath);
+
+				UE_LOG(LogApplePlatformMisc, Log, TEXT("Network path changed (IsConstrained=%d, IsExpensive=%d)"), bIsConstrained, bIsExpensive);
+
+				// cache the path we get
+				if (UE::FApplePlatformMisc::CachedNetworkPath != nullptr)
+				{
+					nw_release(UE::FApplePlatformMisc::CachedNetworkPath);
+				}
+				UE::FApplePlatformMisc::CachedNetworkPath = InPath;
+				nw_retain(UE::FApplePlatformMisc::CachedNetworkPath);
+
 				ENetworkConnectionType ConnectionType = UE::FApplePlatformMisc::CheckNetworkConnectionType();
-				if (ConnectionType != UE::FApplePlatformMisc::LastReportedNetworkType)
+
+				if (FTaskGraphInterface::IsRunning())
 				{
-					UE::FApplePlatformMisc::LastReportedNetworkType = ConnectionType;
-					FFunctionGraphTask::CreateAndDispatchWhenReady([ConnectionType]() {
-						FCoreDelegates::OnNetworkConnectionChanged.Broadcast(ConnectionType);
-					}, TStatId(), nullptr, ENamedThreads::GameThread);
+					if (ConnectionType != UE::FApplePlatformMisc::LastReportedNetworkType)
+					{
+						UE::FApplePlatformMisc::LastReportedNetworkType = ConnectionType;
+						FFunctionGraphTask::CreateAndDispatchWhenReady([ConnectionType]() {
+							FCoreDelegates::OnNetworkConnectionChanged.Broadcast(ConnectionType);
+							FPlatformMisc::SetNetworkConnectionStatus(FPlatformMisc::NetworkConnectionTypeToStatus(ConnectionType));
+						}, TStatId(), nullptr, ENamedThreads::GameThread);
+					}
+
+					// Adding a reference, so path survives to the task body.
+					nw_retain(InPath);
+					FFunctionGraphTask::CreateAndDispatchWhenReady(
+					// nw_path_t is actually a pointer, so captured by value.
+					[InPath] ()
+					{
+						UE::FApplePlatformMisc::NotifyAboutNetworkConnectionCharacteristicsChange(InPath);
+					}, TStatId());
 				}
-				
-				// Adding a reference, so path survives to the task body.
-				nw_retain(InPath);
-				FFunctionGraphTask::CreateAndDispatchWhenReady(
-				// nw_path_t is actually a pointer, so captured by value.
-				[InPath] ()
+				else
 				{
-					UE::FApplePlatformMisc::NotifyAboutNetworkConnectionCharacteristicsChange(InPath);
-				}, TStatId());
+					// Task graph not running yet (early startup). Set the initial status synchronously so it isn't silently dropped.
+					UE::FApplePlatformMisc::LastReportedNetworkType = ConnectionType;
+					FPlatformMisc::SetNetworkConnectionStatus(FPlatformMisc::NetworkConnectionTypeToStatus(ConnectionType));
+				}
 			}
-		}
-	);
-	
-	nw_path_monitor_set_queue(UE::FApplePlatformMisc::PathMonitor, dispatch_get_main_queue());
-	nw_path_monitor_start(UE::FApplePlatformMisc::PathMonitor);
+		);
+
+		nw_path_monitor_set_queue(UE::FApplePlatformMisc::PathMonitor, dispatch_get_main_queue());
+		nw_path_monitor_start(UE::FApplePlatformMisc::PathMonitor);
 
+		// The nw_path_monitor callback will fire shortly after start.
+		// If the task graph is not yet running at that point, the callback sets the initial network status synchronously
+	}
+}
+
+void FApplePlatformMisc::PlatformTearDown()
+{
+	if (UE::FApplePlatformMisc::PathMonitor != nullptr)
+	{
+		nw_path_monitor_cancel(UE::FApplePlatformMisc::PathMonitor);
+		nw_release(UE::FApplePlatformMisc::PathMonitor);
+		UE::FApplePlatformMisc::PathMonitor = nullptr;
+	}
+
+	if (UE::FApplePlatformMisc::CachedNetworkPath != nullptr)
+	{
+		nw_release(UE::FApplePlatformMisc::CachedNetworkPath);
+		UE::FApplePlatformMisc::CachedNetworkPath = nullptr;
+	}
 }
diff --git a/Engine/Source/Runtime/Core/Private/IOS/IOSPlatformMisc.cpp b/Engine/Source/Runtime/Core/Private/IOS/IOSPlatformMisc.cpp
--- a/Engine/Source/Runtime/Core/Private/IOS/IOSPlatformMisc.cpp
+++ b/Engine/Source/Runtime/Core/Private/IOS/IOSPlatformMisc.cpp
@@ -198,6 +198,10 @@ void FIOSPlatformMisc::PlatformInit()
 	FApplePlatformMisc::PlatformInit();
 }
 
+void FIOSPlatformMisc::PlatformTearDown()
+{
+	FApplePlatformMisc::PlatformTearDown();
+}
 
 // Defines the PlatformFeatures module name for iOS, used by PlatformFeatures.h.
 const TCHAR* FIOSPlatformMisc::GetPlatformFeaturesModuleName()
diff --git a/Engine/Source/Runtime/Core/Private/Linux/LinuxPlatformMisc.cpp b/Engine/Source/Runtime/Core/Private/Linux/LinuxPlatformMisc.cpp
--- a/Engine/Source/Runtime/Core/Private/Linux/LinuxPlatformMisc.cpp
+++ b/Engine/Source/Runtime/Core/Private/Linux/LinuxPlatformMisc.cpp
@@ -4,8 +4,177 @@
 	LinuxPlatformMisc.cpp: Linux implementations of misc platform functions
 =============================================================================*/
 
-/* Add back if we overload any functions for GenericPlatform here
 #include "Linux/LinuxPlatformMisc.h"
-#include "Containers/Array.h"
-#include "Containers/UnrealString.h"
-*/
+#include "Logging/LogMacros.h"
+#include "Misc/CoreDelegates.h"
+
+#include <ifaddrs.h>
+#include <net/if.h>
+#include <sys/socket.h>
+#include <linux/netlink.h>
+#include <linux/rtnetlink.h>
+#include <unistd.h>
+
+/** Netlink socket file descriptor for receiving network change events. -1 when not initialized. */
+static int GNetlinkSocket = -1;
+
+/** Last network connection type that was broadcast, used to suppress redundant notifications. */
+static ENetworkConnectionType GLastReportedNetworkType = ENetworkConnectionType::Unknown;
+
+void FLinuxPlatformMisc::PlatformInit()
+{
+	FUnixPlatformMisc::PlatformInit();
+	StartNetworkConnectionMonitoring();
+}
+
+void FLinuxPlatformMisc::PlatformTearDown()
+{
+	StopNetworkConnectionMonitoring();
+	FUnixPlatformMisc::PlatformTearDown();
+}
+
+ENetworkConnectionType FLinuxPlatformMisc::GetNetworkConnectionType()
+{
+	struct ifaddrs* InterfaceAddrs = nullptr;
+	if (getifaddrs(&InterfaceAddrs) != 0)
+	{
+		return ENetworkConnectionType::Unknown;
+	}
+
+	ENetworkConnectionType Result = ENetworkConnectionType::None;
+
+	for (struct ifaddrs* CurrentAddr = InterfaceAddrs; CurrentAddr != nullptr; CurrentAddr = CurrentAddr->ifa_next)
+	{
+		// Skip interfaces that are not up and running
+		if ((CurrentAddr->ifa_flags & (IFF_UP | IFF_RUNNING)) != (IFF_UP | IFF_RUNNING))
+		{
+			continue;
+		}
+
+		// Skip loopback interfaces
+		if (CurrentAddr->ifa_flags & IFF_LOOPBACK)
+		{
+			continue;
+		}
+
+		// Must have a valid address (AF_INET or AF_INET6)
+		if (CurrentAddr->ifa_addr == nullptr)
+		{
+			continue;
+		}
+
+		int Family = CurrentAddr->ifa_addr->sa_family;
+		if (Family != AF_INET && Family != AF_INET6)
+		{
+			continue;
+		}
+
+		// Found a non-loopback interface that is up and running with a valid address
+		Result = ENetworkConnectionType::Ethernet;
+		break;
+	}
+
+	freeifaddrs(InterfaceAddrs);
+	return Result;
+}
+
+void FLinuxPlatformMisc::StartNetworkConnectionMonitoring()
+{
+	check(GNetlinkSocket == -1);
+
+	GNetlinkSocket = socket(AF_NETLINK, SOCK_DGRAM | SOCK_NONBLOCK | SOCK_CLOEXEC, NETLINK_ROUTE);
+	if (GNetlinkSocket == -1)
+	{
+		UE_LOG(LogInit, Warning, TEXT("Failed to create netlink socket for network monitoring (errno=%d)"), errno);
+		return;
+	}
+
+	struct sockaddr_nl Addr;
+	FMemory::Memzero(Addr);
+	Addr.nl_family = AF_NETLINK;
+	Addr.nl_groups = RTMGRP_LINK | RTMGRP_IPV4_IFADDR | RTMGRP_IPV6_IFADDR;
+
+	if (bind(GNetlinkSocket, (struct sockaddr*)&Addr, sizeof(Addr)) < 0)
+	{
+		UE_LOG(LogInit, Warning, TEXT("Failed to bind netlink socket for network monitoring (errno=%d)"), errno);
+		close(GNetlinkSocket);
+		GNetlinkSocket = -1;
+		return;
+	}
+
+	// Set the initial network connection status.
+	GLastReportedNetworkType = GetNetworkConnectionType();
+	SetNetworkConnectionStatus(NetworkConnectionTypeToStatus(GLastReportedNetworkType));
+}
+
+void FLinuxPlatformMisc::StopNetworkConnectionMonitoring()
+{
+	if (GNetlinkSocket != -1)
+	{
+		close(GNetlinkSocket);
+		GNetlinkSocket = -1;
+	}
+}
+
+void FLinuxPlatformMisc::CheckNetworkConnectionEvents()
+{
+	if (GNetlinkSocket == -1)
+	{
+		return;
+	}
+
+	// Non-blocking read: drain any pending netlink messages.
+	// We only care that something changed, not the specific message contents.
+	char Buffer[4096];
+	bool bHasEvents = false;
+
+	while (true)
+	{
+		ssize_t BytesRead = recv(GNetlinkSocket, Buffer, sizeof(Buffer), MSG_DONTWAIT);
+
+		if (BytesRead > 0)
+		{
+			bHasEvents = true;
+			continue;
+		}
+
+		if (BytesRead == 0 || errno == EAGAIN || errno == EWOULDBLOCK)
+		{
+			// No more data available — normal non-blocking termination.
+			break;
+		}
+
+		if (errno == EINTR)
+		{
+			// Interrupted by a signal — retry immediately.
+			continue;
+		}
+
+		if (errno == ENOBUFS || errno == EMSGSIZE)
+		{
+			// ENOBUFS: kernel receive buffer overflowed — netlink messages were silently dropped.
+			// EMSGSIZE: netlink message exceeded the recv buffer — message was truncated.
+			// In both cases, force a status check since we may have missed network change events.
+			bHasEvents = true;
+			break;
+		}
+
+		// Real socket error. Close the socket so we don't keep retrying a broken socket every frame.
+		UE_LOG(LogInit, Warning, TEXT("Netlink recv() failed (errno=%d). Disabling network connection monitoring."), errno);
+		close(GNetlinkSocket);
+		GNetlinkSocket = -1;
+		return;
+	}
+
+	if (bHasEvents)
+	{
+		const ENetworkConnectionType ConnectionType = GetNetworkConnectionType();
+
+		if (ConnectionType != GLastReportedNetworkType)
+		{
+			GLastReportedNetworkType = ConnectionType;
+			FCoreDelegates::OnNetworkConnectionChanged.Broadcast(ConnectionType);
+			SetNetworkConnectionStatus(NetworkConnectionTypeToStatus(ConnectionType));
+		}
+	}
+}
diff --git a/Engine/Source/Runtime/Core/Private/Mac/MacPlatformMisc.cpp b/Engine/Source/Runtime/Core/Private/Mac/MacPlatformMisc.cpp
--- a/Engine/Source/Runtime/Core/Private/Mac/MacPlatformMisc.cpp
+++ b/Engine/Source/Runtime/Core/Private/Mac/MacPlatformMisc.cpp
@@ -647,6 +647,7 @@ void FMacPlatformMisc::PlatformTearDown()
 		[GMacAppInfo.StdErrPipe release];
 	}
 	FPlatformProcess::CeaseBeingFirstInstance();
+	FApplePlatformMisc::PlatformTearDown();
 }
 
 TArray<uint8> FMacPlatformMisc::GetMacAddress()
diff --git a/Engine/Source/Runtime/Core/Public/Apple/ApplePlatformMisc.h b/Engine/Source/Runtime/Core/Public/Apple/ApplePlatformMisc.h
--- a/Engine/Source/Runtime/Core/Public/Apple/ApplePlatformMisc.h
+++ b/Engine/Source/Runtime/Core/Public/Apple/ApplePlatformMisc.h
@@ -42,6 +42,7 @@ struct FDefaultTSDelegateUserPolicy;
 struct CORE_API FApplePlatformMisc : public FGenericPlatformMisc
 {
 	static CORE_API void PlatformInit();
+	static CORE_API void PlatformTearDown();
 
 	static FString GetEnvironmentVariable(const TCHAR* VariableName);
 
diff --git a/Engine/Source/Runtime/Core/Public/GenericPlatform/GenericPlatformMisc.h b/Engine/Source/Runtime/Core/Public/GenericPlatform/GenericPlatformMisc.h
--- a/Engine/Source/Runtime/Core/Public/GenericPlatform/GenericPlatformMisc.h
+++ b/Engine/Source/Runtime/Core/Public/GenericPlatform/GenericPlatformMisc.h
@@ -1686,6 +1686,28 @@ public:
 	 */
 	static CORE_API void SetNetworkConnectionStatus(ENetworkConnectionStatus NewNetworkConnectionStatus);
 
+	/**
+	 * Converts a network connection type to a network connection status.
+	 * Maps None/AirplaneMode to Disabled, all known connection types to Connected.
+	 */
+	static ENetworkConnectionStatus NetworkConnectionTypeToStatus(ENetworkConnectionType ConnectionType)
+	{
+		switch (ConnectionType)
+		{
+		case ENetworkConnectionType::None:
+		case ENetworkConnectionType::AirplaneMode:
+			return ENetworkConnectionStatus::Disabled;
+		case ENetworkConnectionType::Cell:
+		case ENetworkConnectionType::WiFi:
+		case ENetworkConnectionType::WiMAX:
+		case ENetworkConnectionType::Bluetooth:
+		case ENetworkConnectionType::Ethernet:
+		case ENetworkConnectionType::Unknown:
+		default:
+			return ENetworkConnectionStatus::Connected;
+		}
+	}
+
 	/**
 	 * Returns whether WiFi connection is currently active
 	 */
diff --git a/Engine/Source/Runtime/Core/Public/IOS/IOSPlatformMisc.h b/Engine/Source/Runtime/Core/Public/IOS/IOSPlatformMisc.h
--- a/Engine/Source/Runtime/Core/Public/IOS/IOSPlatformMisc.h
+++ b/Engine/Source/Runtime/Core/Public/IOS/IOSPlatformMisc.h
@@ -26,6 +26,7 @@ struct CORE_API FIOSPlatformMisc : public FApplePlatformMisc
 {
     static void PlatformPreInit();
 	static void PlatformInit();
+	static void PlatformTearDown();
     static void PlatformHandleSplashScreen(bool ShowSplashScreen = false);
 	static const TCHAR* GetPlatformFeaturesModuleName();
 
diff --git a/Engine/Source/Runtime/Core/Public/Linux/LinuxPlatformMisc.h b/Engine/Source/Runtime/Core/Public/Linux/LinuxPlatformMisc.h
--- a/Engine/Source/Runtime/Core/Public/Linux/LinuxPlatformMisc.h
+++ b/Engine/Source/Runtime/Core/Public/Linux/LinuxPlatformMisc.h
@@ -12,9 +12,28 @@
  * Linux implementation of the misc OS functions
  */
 
-// A stub for now. Will allow to override any functions if needed
 struct FLinuxPlatformMisc : public FUnixPlatformMisc
 {
+	static CORE_API void PlatformInit();
+	static CORE_API void PlatformTearDown();
+
+	/**
+	 * Returns the current network connection type by checking non-loopback interfaces via getifaddrs().
+	 */
+	static CORE_API ENetworkConnectionType GetNetworkConnectionType();
+
+	/**
+	 * Checks the netlink socket for pending network change events and updates
+	 * connection status if a change occurred. Should be called from PumpMessages().
+	 */
+	static CORE_API void CheckNetworkConnectionEvents();
+
+private:
+	/** Creates a NETLINK_ROUTE socket to listen for interface and address changes. */
+	static void StartNetworkConnectionMonitoring();
+
+	/** Closes the netlink socket. */
+	static void StopNetworkConnectionMonitoring();
 };
 
 typedef FLinuxPlatformMisc FPlatformMisc;
