Added FMallocBinned4 support to Unix platform

#Synced-CL 52054134

diff --git a/Engine/Source/Runtime/Core/Private/Unix/UnixPlatformMemory.cpp b/Engine/Source/Runtime/Core/Private/Unix/UnixPlatformMemory.cpp
--- a/Engine/Source/Runtime/Core/Private/Unix/UnixPlatformMemory.cpp
+++ b/Engine/Source/Runtime/Core/Private/Unix/UnixPlatformMemory.cpp
@@ -20,6 +20,7 @@
 #include "HAL/MallocBinned.h"
 #include "HAL/MallocBinned2.h"
 #include "HAL/MallocBinned3.h"
+#include "HAL/MallocBinned4.h"
 #include "HAL/MallocReplayProxy.h"
 #include "HAL/MallocStomp.h"
 #include "HAL/PlatformMallocCrash.h"
@@ -223,6 +224,12 @@ class FMalloc* FUnixPlatformMemory::BaseAllocator()
 				break;
 			}
 
+			if (FCStringAnsi::Stricmp(Arg, "-binnedmalloc4") == 0)
+			{
+				AllocatorToUse = EMemoryAllocatorToUse::Binned4;
+				break;
+			}
+
 			if (FCStringAnsi::Stricmp(Arg, "-fullcrashcallstack") == 0)
 			{
 				GFullCrashCallstack = true;
@@ -353,6 +360,10 @@ class FMalloc* FUnixPlatformMemory::BaseAllocator()
 		Allocator = new FMallocBinned3();
 		break;
 
+	case EMemoryAllocatorToUse::Binned4:
+		Allocator = new FMallocBinned4();
+		break;
+
 	default:	// intentional fall-through
 	case EMemoryAllocatorToUse::Binned:
 		Allocator = new FMallocBinned(FPlatformMemory::GetConstants().BinnedPageSize & MAX_uint32, 0x100000000);
diff --git a/Engine/Source/Runtime/Core/Public/GenericPlatform/GenericPlatformMemory.h b/Engine/Source/Runtime/Core/Public/GenericPlatform/GenericPlatformMemory.h
--- a/Engine/Source/Runtime/Core/Public/GenericPlatform/GenericPlatformMemory.h
+++ b/Engine/Source/Runtime/Core/Public/GenericPlatform/GenericPlatformMemory.h
@@ -269,6 +269,7 @@ struct FGenericPlatformMemory
 		Binned, // Older binned malloc
 		Binned2, // Newer binned malloc
 		Binned3, // Newer VM-based binned malloc, 64 bit only
+		Binned4, // New experimental binned allocator
 		Platform, // Custom platform specific allocator
 		Mimalloc, // mimalloc
 		Libpas, // libpas
