Replaced the bHasPayload flag in the output stream by folding the flag into the serialized function name.

#Synced-CL 52250244

diff --git a/Engine/Source/Runtime/Core/Public/UObject/FortniteMainBranchObjectVersions.inl b/Engine/Source/Runtime/Core/Public/UObject/FortniteMainBranchObjectVersions.inl
--- a/Engine/Source/Runtime/Core/Public/UObject/FortniteMainBranchObjectVersions.inl
+++ b/Engine/Source/Runtime/Core/Public/UObject/FortniteMainBranchObjectVersions.inl
@@ -796,8 +796,11 @@ FORTNITE_MAIN_VERSION(AddDevNotesToFText, 260)
 // Added default color override for PCG Mesh Sampler and fix bounds in PCG Subdivide Spline
 FORTNITE_MAIN_VERSION(PCGMeshSamplerDefaultColorPCGSubdivideSplineBoundsFixes, 261)
 
-// Serialisable dynamic delegate payloads
+// Serializable dynamic delegate payloads
 FORTNITE_MAIN_VERSION(DynamicDelegatePayloads, 262)
 
+// Compact representation of serializable dynamic delegate payloads which avoids shifting binary data
+FORTNITE_MAIN_VERSION(CompactDynamicDelegatePayloads, 263)
+
 // -----<new versions can be added above this line>-------------------------------------------------
 #undef FORTNITE_MAIN_VERSION
diff --git a/Engine/Source/Runtime/Core/Public/UObject/ScriptDelegates.h b/Engine/Source/Runtime/Core/Public/UObject/ScriptDelegates.h
--- a/Engine/Source/Runtime/Core/Public/UObject/ScriptDelegates.h
+++ b/Engine/Source/Runtime/Core/Public/UObject/ScriptDelegates.h
@@ -51,6 +51,10 @@
 	#include "UObject/WeakObjectPtrTemplatesFwd.h"
 #endif
 
+// A suffix appended to the serialized function name to declare that a payload was written
+// after it.  This is used instead of an explicit flag to avoid impacting the existing binary format.
+#define UE_DYNAMIC_DELEGATE_PAYLOAD_FUNCTION_NAME_SUFFIX TEXT("[bHasPayload]")
+
 namespace UE::Core::Private
 {
 	template <typename InThreadSafetyMode>
@@ -613,17 +617,48 @@ public:
 	{
 		FReadAccessScope ReadScope = D.GetReadAccessScope();
 
-		Ar << D.Object << D.FunctionName;
-
 		Ar.UsingCustomVersion(FFortniteMainBranchObjectVersion::GUID);
 
+		// The original implementation serialized a bool to represent the presence of a payload.
+		// The new implementation folds the flag into the function name so that the binary layout is unaffected for dynamic delegates without payloads
+		bool bSerializePayloadData         = Ar.CustomVer(FFortniteMainBranchObjectVersion::GUID) >= FFortniteMainBranchObjectVersion::DynamicDelegatePayloads;
+		bool bPayloadDataUsesCompactFormat = Ar.CustomVer(FFortniteMainBranchObjectVersion::GUID) >= FFortniteMainBranchObjectVersion::CompactDynamicDelegatePayloads;
+		if (!bPayloadDataUsesCompactFormat)
+		{
+			Ar << D.Object << D.FunctionName;
+		}
+
 		bool bHasPayload = false;
-		if (Ar.CustomVer(FFortniteMainBranchObjectVersion::GUID) >= FFortniteMainBranchObjectVersion::DynamicDelegatePayloads)
+		if (bSerializePayloadData)
 		{
 #if UE_USE_DYNAMIC_DELEGATE_PAYLOADS
 			bHasPayload = !!D.Payload;
 #endif
-			Ar << bHasPayload;
+			if (!bPayloadDataUsesCompactFormat)
+			{
+				Ar << bHasPayload;
+			}
+		}
+
+		if (bPayloadDataUsesCompactFormat)
+		{
+			FName FunctionNameWithOptionalPayloadFlag = D.FunctionName;
+			if (Ar.IsSaving())
+			{
+				if (bHasPayload)
+				{
+					FunctionNameWithOptionalPayloadFlag = *(FunctionNameWithOptionalPayloadFlag.ToString() + UE_DYNAMIC_DELEGATE_PAYLOAD_FUNCTION_NAME_SUFFIX);
+				}
+			}
+
+			Ar << D.Object << FunctionNameWithOptionalPayloadFlag;
+
+			if (Ar.IsLoading())
+			{
+				FString FunctionNameWithOptionalPayloadFlagStr = FunctionNameWithOptionalPayloadFlag.ToString();
+				bHasPayload = FunctionNameWithOptionalPayloadFlagStr.RemoveFromEnd(UE_DYNAMIC_DELEGATE_PAYLOAD_FUNCTION_NAME_SUFFIX);
+				D.FunctionName = *FunctionNameWithOptionalPayloadFlagStr;
+			}
 		}
 
 		if (bHasPayload)
@@ -664,49 +699,98 @@ public:
 
 		FArchive& UnderlyingAr = Slot.GetUnderlyingArchive();
 
-		FStructuredArchive::FRecord Record = Slot.EnterRecord();
-		Record << SA_VALUE(TEXT("Object"), D.Object) << SA_VALUE(TEXT("FunctionName"),D.FunctionName);
-
 		UnderlyingAr.UsingCustomVersion(FFortniteMainBranchObjectVersion::GUID);
 
-		if (UnderlyingAr.CustomVer(FFortniteMainBranchObjectVersion::GUID) >= FFortniteMainBranchObjectVersion::DynamicDelegatePayloads)
+		bool bSerializePayloadData         = UnderlyingAr.CustomVer(FFortniteMainBranchObjectVersion::GUID) >= FFortniteMainBranchObjectVersion::DynamicDelegatePayloads;
+		bool bPayloadDataUsesCompactFormat = UnderlyingAr.CustomVer(FFortniteMainBranchObjectVersion::GUID) >= FFortniteMainBranchObjectVersion::CompactDynamicDelegatePayloads;
+
+		FStructuredArchive::FRecord Record = Slot.EnterRecord();
+
+		bool bHasPayload = false;
+		if (bSerializePayloadData)
 		{
-			bool bHasPayload = false;
 #if UE_USE_DYNAMIC_DELEGATE_PAYLOADS
 			bHasPayload = !!D.Payload;
 #endif
+		}
 
-			// TryEnterField writes the Payload field only when bHasPayload is true (saving),
-			// and on load returns a slot only if the field exists in the archive.
-			// The payload is always binary — the slot hooks use FArchiveUObjectFromStructuredArchive
-			// (in CoreUObject) to wrap the slot as a binary archive, producing a base-64 blob in text formats.
-			if (TOptional<FStructuredArchiveSlot> PayloadSlot = Record.TryEnterField(TEXT("Payload"), bHasPayload))
+		// For text-format archives, payload presence is expressed structurally via TryEnterField
+		// and the function name is always clean.  For binary-format archives, we use the function
+		// name suffix trick (v263+) to avoid shifting binary data.
+		if (bPayloadDataUsesCompactFormat && !UnderlyingAr.IsTextFormat())
+		{
+			// Binary structured path — same suffix approach as the binary operator<<
+			FName FunctionNameWithOptionalPayloadFlag = D.FunctionName;
+			if (UnderlyingAr.IsSaving())
 			{
-				if (UnderlyingAr.IsLoading())
+				if (bHasPayload)
 				{
-#if UE_USE_DYNAMIC_DELEGATE_PAYLOADS
-					if (ensureMsgf(UE::Core::Private::GDelegatePayloadDeserializeSlotFunc, TEXT("Trying to deserialize a dynamic delegate payload without CoreUObject's deserialize factory being registered")))
-					{
-						D.Payload = UE::Core::Private::GDelegatePayloadDeserializeSlotFunc(*PayloadSlot);
-					}
-					else
-					{
-						D.Unbind();
-						// TODO: drain payload bytes once FDynamicDelegatePropertyPayload exists (Step 3)
-					}
-#else
-					ensureMsgf(false, TEXT("Trying to deserialize a dynamic delegate payload when payloads are disabled"));
-					D.Unbind();
-					// TODO: drain bytes (Step 3)
-#endif
+					FunctionNameWithOptionalPayloadFlag = *(FunctionNameWithOptionalPayloadFlag.ToString() + UE_DYNAMIC_DELEGATE_PAYLOAD_FUNCTION_NAME_SUFFIX);
 				}
+			}
+
+			Record << SA_VALUE(TEXT("Object"), D.Object) << SA_VALUE(TEXT("FunctionName"), FunctionNameWithOptionalPayloadFlag);
+
+			if (UnderlyingAr.IsLoading())
+			{
+				FString FunctionNameWithOptionalPayloadFlagStr = FunctionNameWithOptionalPayloadFlag.ToString();
+				bHasPayload = FunctionNameWithOptionalPayloadFlagStr.RemoveFromEnd(UE_DYNAMIC_DELEGATE_PAYLOAD_FUNCTION_NAME_SUFFIX);
+				D.FunctionName = *FunctionNameWithOptionalPayloadFlagStr;
+			}
+		}
+		else
+		{
+			// Text format or pre-v263: clean function name, TryEnterField for payload presence
+			Record << SA_VALUE(TEXT("Object"), D.Object) << SA_VALUE(TEXT("FunctionName"), D.FunctionName);
+		}
+
+		// Serialize payload data
+		auto SerializePayload = [&UnderlyingAr, &D](FStructuredArchiveSlot PayloadSlot)
+		{
+			if (UnderlyingAr.IsLoading())
+			{
 #if UE_USE_DYNAMIC_DELEGATE_PAYLOADS
+				if (ensureMsgf(UE::Core::Private::GDelegatePayloadDeserializeSlotFunc, TEXT("Trying to deserialize a dynamic delegate payload without CoreUObject's deserialize factory being registered")))
+				{
+					D.Payload = UE::Core::Private::GDelegatePayloadDeserializeSlotFunc(PayloadSlot);
+				}
 				else
 				{
-					UE::Core::Private::GDelegatePayloadSerializeSlotFunc(*PayloadSlot, *D.Payload);
+					D.Unbind();
+					// TODO: drain payload bytes once FDynamicDelegatePropertyPayload exists (Step 3)
 				}
+#else
+				ensureMsgf(false, TEXT("Trying to deserialize a dynamic delegate payload when payloads are disabled"));
+				D.Unbind();
+				// TODO: drain bytes (Step 3)
 #endif
 			}
+#if UE_USE_DYNAMIC_DELEGATE_PAYLOADS
+			else
+			{
+				UE::Core::Private::GDelegatePayloadSerializeSlotFunc(PayloadSlot, *D.Payload);
+			}
+#endif
+		};
+
+		if (bSerializePayloadData)
+		{
+			if (bPayloadDataUsesCompactFormat && !UnderlyingAr.IsTextFormat())
+			{
+				// Binary v263+: bHasPayload already determined from suffix
+				if (bHasPayload)
+				{
+					SerializePayload(Record.EnterField(TEXT("Payload")));
+				}
+			}
+			else
+			{
+				// Text format or v262: TryEnterField expresses presence structurally
+				if (TOptional<FStructuredArchiveSlot> PayloadSlot = Record.TryEnterField(TEXT("Payload"), bHasPayload))
+				{
+					SerializePayload(*PayloadSlot);
+				}
+			}
 		}
 	}
 
