[Core/HAL] Fix: FSharedEventRef not respecting EEventMode passed in to the constructor

#Synced-CL 48443255

diff --git a/Engine/Source/Runtime/Core/Private/HAL/ThreadingBase.cpp b/Engine/Source/Runtime/Core/Private/HAL/ThreadingBase.cpp
--- a/Engine/Source/Runtime/Core/Private/HAL/ThreadingBase.cpp
+++ b/Engine/Source/Runtime/Core/Private/HAL/ThreadingBase.cpp
@@ -127,7 +127,7 @@ FTaskTagScope::FTaskTagScope(ETaskTag InTag) : Tag(InTag)
 			//Try to catch other Threads that already opened a non parallel scope
 			ActiveNamedThreads.store(int32(ETaskTag::ENone));
 		}
-		checkf(IsOK || IsEngineExitRequested(), TEXT("Only Scopes tagged with ETaskTag::EParallelThread can be tagged multiple times. ActiveNamedThreads(%x) cannot be tagged multiple times in the same callstack with the current Tag(%x)"), int32{EnumToUnderlyingType(OldTag)}, int32{EnumToUnderlyingType(Tag)});
+		checkf(IsOK || IsEngineExitRequested(), TEXT("Only Scopes tagged with ETaskTag::EParallelThread can be tagged multiple times. ActiveNamedThreads(%x) cannot be tagged multiple times in the same callstack with the current Tag(%x)"), OldTag, Tag);
 	}
 	
 	ParentTag = ActiveTaskTag;
@@ -139,19 +139,19 @@ FTaskTagScope::FTaskTagScope(ETaskTag InTag) : Tag(InTag)
 	{
 		if (EnumHasAllFlags(Tag, ETaskTag::EParallelRenderingThread))
 		{
-			checkf(IsInRenderingThread(), TEXT("ETaskTag::EParallelRenderingThread can only be retagged if they are in a parallel for on the RenderingThread or not tagged check the ActiveNamedThreads(%x)"), int32{EnumToUnderlyingType(FTaskTagScope::GetCurrentTag())});
+			checkf(IsInRenderingThread(), TEXT("ETaskTag::EParallelRenderingThread can only be retagged if they are in a parallel for on the RenderingThread or not tagged check the ActiveNamedThreads(%x)"), FTaskTagScope::GetCurrentTag());
 		}
 
 		if (EnumHasAllFlags(Tag, ETaskTag::EParallelGameThread))
 		{
-			checkf(IsInGameThread(), TEXT("ETaskTag::EParallelGameThread can only be retagged if they are in a parallel for on the GameThread or not tagged check the ActiveNamedThreads(%x)"), int32{EnumToUnderlyingType(FTaskTagScope::GetCurrentTag())});
+			checkf(IsInGameThread(), TEXT("ETaskTag::EParallelGameThread can only be retagged if they are in a parallel for on the GameThread or not tagged check the ActiveNamedThreads(%x)"), FTaskTagScope::GetCurrentTag());
 		}
 	}
 }
 
 FTaskTagScope::~FTaskTagScope()
 {
-	checkf(TagOnlyIfNone || ActiveTaskTag == Tag, TEXT("ActiveTaskTag(%x) corrupted needs to be Tag(%x)"), int32{EnumToUnderlyingType(FTaskTagScope::GetCurrentTag())}, int32{EnumToUnderlyingType(Tag)});
+	checkf(TagOnlyIfNone || ActiveTaskTag == Tag, TEXT("ActiveTaskTag(%x) corrupted needs to be Tag(%x)"), FTaskTagScope::GetCurrentTag(), Tag);
 	if (!TagOnlyIfNone || ParentTag == ETaskTag::ENone || ParentTag == ETaskTag::EWorkerThread)
 	{
 		ActiveTaskTag = ParentTag;
@@ -162,7 +162,7 @@ FTaskTagScope::~FTaskTagScope()
 		ETaskTag NamedThreadBits = (Tag & ETaskTag::ENamedThreadBits);
 		static_assert(sizeof(ETaskTag) == sizeof(int32), "EnumSize must match interlockedAnd");
 		ETaskTag OldTag = ETaskTag(ActiveNamedThreads.fetch_and(int32(~NamedThreadBits)));
-		checkf((OldTag & NamedThreadBits) == (NamedThreadBits) || IsEngineExitRequested(), TEXT("Currently active Threads(%x) got corrupted check the ActiveNamedThreads(%x)"), int32{EnumToUnderlyingType(OldTag)}, int32{EnumToUnderlyingType(FTaskTagScope::GetCurrentTag())});
+		checkf((OldTag & NamedThreadBits) == (NamedThreadBits) || IsEngineExitRequested(), TEXT("Currently active Threads(%x) got corrupted check the ActiveNamedThreads(%x)"), OldTag, FTaskTagScope::GetCurrentTag());
 	}
 
 	//prolong the scope of the GT for static variable destructors
@@ -198,7 +198,7 @@ PRAGMA_DISABLE_DEPRECATION_WARNINGS
 		{
 			const uint32 CurrentThreadId = FPlatformTLS::GetCurrentThreadId();
 			bool oldValue = CurrentThreadId == GGameThreadId;
-			ensureMsgf(oldValue == newValue, TEXT("oldValue(%i) newValue(%i) If this check fails make sure that there is a FTaskTagScope(ETaskTag::EGameThread) as deep as possible on the current callstack, you can see the current value in ActiveNamedThreads(%x)"), oldValue, newValue, int32{EnumToUnderlyingType(FTaskTagScope::GetCurrentTag())});
+			ensureMsgf(oldValue == newValue, TEXT("oldValue(%i) newValue(%i) If this check fails make sure that there is a FTaskTagScope(ETaskTag::EGameThread) as deep as possible on the current callstack, you can see the current value in ActiveNamedThreads(%x)"), oldValue, newValue, FTaskTagScope::GetCurrentTag());
 			newValue = oldValue;
 		}
 #endif
@@ -233,7 +233,7 @@ PRAGMA_DISABLE_DEPRECATION_WARNINGS
 	if (!UE::Tasks::Private::IsThreadRetractingTask())
 	{
 		bool oldValue = GSlateLoadingThreadId != 0 && FPlatformTLS::GetCurrentThreadId() == GSlateLoadingThreadId;
-		ensureMsgf(oldValue == newValue, TEXT("oldValue(%i) newValue(%i) If this check fails make sure that there is a FTaskTagScope(ETaskTag::ESlateThread) as as deep as possible on the current callstack, you can see the current value in ActiveNamedThreads(%x)"), oldValue, newValue, int32{EnumToUnderlyingType(FTaskTagScope::GetCurrentTag())});
+		ensureMsgf(oldValue == newValue, TEXT("oldValue(%i) newValue(%i) If this check fails make sure that there is a FTaskTagScope(ETaskTag::ESlateThread) as as deep as possible on the current callstack, you can see the current value in ActiveNamedThreads(%x)"), oldValue, newValue, FTaskTagScope::GetCurrentTag());
 		newValue = oldValue;
 	}
 #endif
@@ -265,7 +265,7 @@ PRAGMA_DISABLE_DEPRECATION_WARNINGS
 #if !UE_BUILD_SHIPPING && !UE_BUILD_TEST
 	{
 		bool oldValue = FPlatformTLS::GetCurrentThreadId() == GRenderThreadId;
-		ensureMsgf(oldValue == newValue, TEXT("oldValue(%i) newValue(%i) If this check fails make sure that there is a FTaskTagScope(ETaskTag::ERenderingThread) as deep as possible on the current callstack, you can see the current value in ActiveNamedThreads(%x)"), oldValue, newValue, int32{EnumToUnderlyingType(FTaskTagScope::GetCurrentTag())});
+		ensureMsgf(oldValue == newValue, TEXT("oldValue(%i) newValue(%i) If this check fails make sure that there is a FTaskTagScope(ETaskTag::ERenderingThread) as deep as possible on the current callstack, you can see the current value in ActiveNamedThreads(%x)"), oldValue, newValue, FTaskTagScope::GetCurrentTag());
 		newValue = oldValue;
 	}
 #endif
@@ -285,7 +285,7 @@ PRAGMA_DISABLE_DEPRECATION_WARNINGS
 	{
 		const uint32 CurrentThreadId = FPlatformTLS::GetCurrentThreadId();
 		bool oldValue = (GRenderThreadId == 0) ? (CurrentThreadId == GGameThreadId) || FTaskTagScope::IsRunningDuringStaticInit() : (CurrentThreadId == GRenderThreadId);
-		ensureMsgf(oldValue == newValue, TEXT("oldValue(%i) newValue(%i) If this check fails make sure that there is a FTaskTagScope(ETaskTag::ERenderingThread) as deep as possible on the current callstack, you can see the current value in ActiveNamedThreads(%x)"), oldValue, newValue, int32{EnumToUnderlyingType(FTaskTagScope::GetCurrentTag())});
+		ensureMsgf(oldValue == newValue, TEXT("oldValue(%i) newValue(%i) If this check fails make sure that there is a FTaskTagScope(ETaskTag::ERenderingThread) as deep as possible on the current callstack, you can see the current value in ActiveNamedThreads(%x)"), oldValue, newValue, FTaskTagScope::GetCurrentTag());
 		newValue = oldValue;
 	}
 #endif
@@ -322,7 +322,7 @@ PRAGMA_DISABLE_DEPRECATION_WARNINGS
 	{
 		const uint32 CurrentThreadId = FPlatformTLS::GetCurrentThreadId();
 		bool oldValue = (GRenderThreadId == 0) ?  true : CurrentThreadId != GGameThreadId;
-		ensureMsgf(oldValue == newValue, TEXT("oldValue(%i) newValue(%i) If this check fails make sure that there is a FTaskTagScope(ETaskTag::EParallelRenderingThread) as deep as possible on the current callstack, you can see the current value in ActiveNamedThreads(%x)"), oldValue, newValue, int32{EnumToUnderlyingType(FTaskTagScope::GetCurrentTag())});
+		ensureMsgf(oldValue == newValue, TEXT("oldValue(%i) newValue(%i) If this check fails make sure that there is a FTaskTagScope(ETaskTag::EParallelRenderingThread) as deep as possible on the current callstack, you can see the current value in ActiveNamedThreads(%x)"), oldValue, newValue, FTaskTagScope::GetCurrentTag());
 		newValue = oldValue;
 	}
 #endif
@@ -849,8 +849,20 @@ FEventRef::~FEventRef()
 -----------------------------------------------------------------------------*/
 
 FSharedEventRef::FSharedEventRef(EEventMode Mode  /* = EEventMode::AutoReset */)
-	: Ptr(TLazySingleton<TEventPool<EEventMode::AutoReset>>::Get().GetRawEvent(),
-		[](FEvent* Event) { TLazySingleton<TEventPool<EEventMode::AutoReset>>::Get().ReturnRawEvent(Event); })
+	: Ptr(Mode == EEventMode::ManualReset
+			? TLazySingleton<TEventPool<EEventMode::ManualReset>>::Get().GetRawEvent()
+			: TLazySingleton<TEventPool<EEventMode::AutoReset>>::Get().GetRawEvent(),
+		[](FEvent* Event)
+		{
+			if (Event->IsManualReset())
+			{
+				TLazySingleton<TEventPool<EEventMode::ManualReset>>::Get().ReturnRawEvent(Event);
+			}
+			else
+			{
+				TLazySingleton<TEventPool<EEventMode::AutoReset>>::Get().ReturnRawEvent(Event);
+			}
+		})
 {
 }
 
@@ -1575,19 +1587,7 @@ FTlsAutoCleanup* FThreadSingletonInitializer::Inject(FTlsAutoCleanup* Instance,
 
 void FTlsAutoCleanup::Register()
 {
-	// Destruction of TLS values need to occur in the reverse order of initialization, as thread_local does.
-	struct FTlsContainer : public TArray<TUniquePtr<FTlsAutoCleanup>>
-	{
-		~FTlsContainer()
-		{
-			for (int32 Index = Num() - 1; Index >= 0; Index--)
-			{
-				(*this)[Index].Reset();
-			}
-		}
-	};
-
-	static thread_local FTlsContainer TlsInstances;
+	static thread_local TArray<TUniquePtr<FTlsAutoCleanup>> TlsInstances;
 	TlsInstances.Add(TUniquePtr<FTlsAutoCleanup>(this));
 }
 
diff --git a/Engine/Source/Runtime/Core/Tests/HAL/SharedEventRefTest.cpp b/Engine/Source/Runtime/Core/Tests/HAL/SharedEventRefTest.cpp
new file mode 100644
--- /dev/null
+++ b/Engine/Source/Runtime/Core/Tests/HAL/SharedEventRefTest.cpp
@@ -0,0 +1,113 @@
+// Copyright Epic Games, Inc. All Rights Reserved.
+
+#if WITH_TESTS
+
+	#include "HAL/Event.h"
+	#include "Tasks/Task.h"
+	#include "Tests/TestHarnessAdapter.h"
+
+namespace UE
+{
+	TEST_CASE_NAMED(FSharedEventRefTest, "Core::HAL::FSharedEventRef", "[Core][HAL]")
+	{
+		SECTION("Default constructor (EEventMode::AutoReset)")
+		{
+			FSharedEventRef Event; // EEventMode::AutoReset (default)
+
+			CHECK_FALSE(Event->IsManualReset());
+
+			// Initial state: unsignaled
+			CHECK_FALSE(Event->Wait(0));
+
+			// Wait with timeout: returns false and respects the timeout duration
+			const double PreWait = FPlatformTime::Seconds();
+			CHECK_FALSE(Event->Wait(32));
+			CHECK((FPlatformTime::Seconds() - PreWait) * 1000.0 >= 16.0);
+
+			// After Trigger: first Wait returns true, then the signal is consumed (auto-reset)
+			Event->Trigger();
+			CHECK(Event->Wait(0));
+			CHECK_FALSE(Event->Wait(0));
+
+			// Cross-thread: signal arrives from a background task while blocked in Wait
+			Tasks::Launch(UE_SOURCE_LOCATION, [&Event] {
+				FPlatformProcess::SleepNoStats(0.02f);
+				Event->Trigger();
+			});
+			const double CrossThreadPreWait = FPlatformTime::Seconds();
+			CHECK(Event->Wait(5000));
+			CHECK((FPlatformTime::Seconds() - CrossThreadPreWait) * 1000.0 >= 16.0);
+			CHECK_FALSE(Event->Wait(0)); // auto-reset: signal was consumed by the Wait above
+		}
+
+		SECTION("Explicit constructor (EEventMode::AutoReset)")
+		{
+			FSharedEventRef Event(EEventMode::AutoReset);
+
+			CHECK_FALSE(Event->IsManualReset());
+
+			// Initial state: unsignaled
+			CHECK_FALSE(Event->Wait(0));
+
+			// Wait with timeout: returns false and respects the timeout duration
+			const double PreWait = FPlatformTime::Seconds();
+			CHECK_FALSE(Event->Wait(32));
+			CHECK((FPlatformTime::Seconds() - PreWait) * 1000.0 >= 16.0);
+
+			// After Trigger: first Wait returns true, then the signal is consumed (auto-reset)
+			Event->Trigger();
+			CHECK(Event->Wait(0));
+			CHECK_FALSE(Event->Wait(0));
+
+			// Cross-thread: signal arrives from a background task while blocked in Wait
+			Tasks::Launch(UE_SOURCE_LOCATION, [&Event] {
+				FPlatformProcess::SleepNoStats(0.02f);
+				Event->Trigger();
+			});
+			const double CrossThreadPreWait = FPlatformTime::Seconds();
+			CHECK(Event->Wait(5000));
+			CHECK((FPlatformTime::Seconds() - CrossThreadPreWait) * 1000.0 >= 16.0);
+			CHECK_FALSE(Event->Wait(0)); // auto-reset: signal was consumed by the Wait above
+		}
+
+		SECTION("Explicit constructor (EEventMode::ManualReset)")
+		{
+			FSharedEventRef Event(EEventMode::ManualReset);
+
+			CHECK(Event->IsManualReset());
+
+			// Initial state: unsignaled
+			CHECK_FALSE(Event->Wait(0));
+
+			// Wait with timeout: returns false and respects the timeout duration
+			const double PreWait = FPlatformTime::Seconds();
+			CHECK_FALSE(Event->Wait(32));
+			CHECK((FPlatformTime::Seconds() - PreWait) * 1000.0 >= 16.0);
+
+			// After Trigger: signal persists across multiple Wait calls (manual reset)
+			Event->Trigger();
+			CHECK(Event->Wait(0));
+			CHECK(Event->Wait(0)); // still signaled
+			CHECK(Event->Wait(0)); // still signaled
+
+			// After Reset: unsignaled again
+			Event->Reset();
+			CHECK_FALSE(Event->Wait(0));
+
+			// Cross-thread: signal arrives from a background task while blocked in Wait
+			Tasks::Launch(UE_SOURCE_LOCATION, [&Event] {
+				FPlatformProcess::SleepNoStats(0.02f);
+				Event->Trigger();
+			});
+			const double CrossThreadPreWait = FPlatformTime::Seconds();
+			CHECK(Event->Wait(5000));
+			CHECK((FPlatformTime::Seconds() - CrossThreadPreWait) * 1000.0 >= 16.0);
+			CHECK(Event->Wait(0)); // still signaled — manual reset does not consume on Wait
+
+			// Cleanup
+			Event->Reset();
+		}
+	}
+} // namespace UE
+
+#endif // WITH_TESTS
