[Backout] - CL52290555
#rnx
Original CL Desc
-----------------------------------------------------------------
[Core/HAL] Fix: FSharedEventRef not respecting EEventMode passed in to the constructor

#Synced-CL 48443255

diff --git a/Engine/Source/Runtime/Core/Private/HAL/ThreadingBase.cpp b/Engine/Source/Runtime/Core/Private/HAL/ThreadingBase.cpp
--- a/Engine/Source/Runtime/Core/Private/HAL/ThreadingBase.cpp
+++ b/Engine/Source/Runtime/Core/Private/HAL/ThreadingBase.cpp
@@ -127,7 +127,7 @@ FTaskTagScope::FTaskTagScope(ETaskTag InTag) : Tag(InTag)
 			//Try to catch other Threads that already opened a non parallel scope
 			ActiveNamedThreads.store(int32(ETaskTag::ENone));
 		}
-		checkf(IsOK || IsEngineExitRequested(), TEXT("Only Scopes tagged with ETaskTag::EParallelThread can be tagged multiple times. ActiveNamedThreads(%x) cannot be tagged multiple times in the same callstack with the current Tag(%x)"), OldTag, Tag);
+		checkf(IsOK || IsEngineExitRequested(), TEXT("Only Scopes tagged with ETaskTag::EParallelThread can be tagged multiple times. ActiveNamedThreads(%x) cannot be tagged multiple times in the same callstack with the current Tag(%x)"), int32{EnumToUnderlyingType(OldTag)}, int32{EnumToUnderlyingType(Tag)});
 	}
 	
 	ParentTag = ActiveTaskTag;
@@ -139,19 +139,19 @@ FTaskTagScope::FTaskTagScope(ETaskTag InTag) : Tag(InTag)
 	{
 		if (EnumHasAllFlags(Tag, ETaskTag::EParallelRenderingThread))
 		{
-			checkf(IsInRenderingThread(), TEXT("ETaskTag::EParallelRenderingThread can only be retagged if they are in a parallel for on the RenderingThread or not tagged check the ActiveNamedThreads(%x)"), FTaskTagScope::GetCurrentTag());
+			checkf(IsInRenderingThread(), TEXT("ETaskTag::EParallelRenderingThread can only be retagged if they are in a parallel for on the RenderingThread or not tagged check the ActiveNamedThreads(%x)"), int32{EnumToUnderlyingType(FTaskTagScope::GetCurrentTag())});
 		}
 
 		if (EnumHasAllFlags(Tag, ETaskTag::EParallelGameThread))
 		{
-			checkf(IsInGameThread(), TEXT("ETaskTag::EParallelGameThread can only be retagged if they are in a parallel for on the GameThread or not tagged check the ActiveNamedThreads(%x)"), FTaskTagScope::GetCurrentTag());
+			checkf(IsInGameThread(), TEXT("ETaskTag::EParallelGameThread can only be retagged if they are in a parallel for on the GameThread or not tagged check the ActiveNamedThreads(%x)"), int32{EnumToUnderlyingType(FTaskTagScope::GetCurrentTag())});
 		}
 	}
 }
 
 FTaskTagScope::~FTaskTagScope()
 {
-	checkf(TagOnlyIfNone || ActiveTaskTag == Tag, TEXT("ActiveTaskTag(%x) corrupted needs to be Tag(%x)"), FTaskTagScope::GetCurrentTag(), Tag);
+	checkf(TagOnlyIfNone || ActiveTaskTag == Tag, TEXT("ActiveTaskTag(%x) corrupted needs to be Tag(%x)"), int32{EnumToUnderlyingType(FTaskTagScope::GetCurrentTag())}, int32{EnumToUnderlyingType(Tag)});
 	if (!TagOnlyIfNone || ParentTag == ETaskTag::ENone || ParentTag == ETaskTag::EWorkerThread)
 	{
 		ActiveTaskTag = ParentTag;
@@ -162,7 +162,7 @@ FTaskTagScope::~FTaskTagScope()
 		ETaskTag NamedThreadBits = (Tag & ETaskTag::ENamedThreadBits);
 		static_assert(sizeof(ETaskTag) == sizeof(int32), "EnumSize must match interlockedAnd");
 		ETaskTag OldTag = ETaskTag(ActiveNamedThreads.fetch_and(int32(~NamedThreadBits)));
-		checkf((OldTag & NamedThreadBits) == (NamedThreadBits) || IsEngineExitRequested(), TEXT("Currently active Threads(%x) got corrupted check the ActiveNamedThreads(%x)"), OldTag, FTaskTagScope::GetCurrentTag());
+		checkf((OldTag & NamedThreadBits) == (NamedThreadBits) || IsEngineExitRequested(), TEXT("Currently active Threads(%x) got corrupted check the ActiveNamedThreads(%x)"), int32{EnumToUnderlyingType(OldTag)}, int32{EnumToUnderlyingType(FTaskTagScope::GetCurrentTag())});
 	}
 
 	//prolong the scope of the GT for static variable destructors
@@ -198,7 +198,7 @@ PRAGMA_DISABLE_DEPRECATION_WARNINGS
 		{
 			const uint32 CurrentThreadId = FPlatformTLS::GetCurrentThreadId();
 			bool oldValue = CurrentThreadId == GGameThreadId;
-			ensureMsgf(oldValue == newValue, TEXT("oldValue(%i) newValue(%i) If this check fails make sure that there is a FTaskTagScope(ETaskTag::EGameThread) as deep as possible on the current callstack, you can see the current value in ActiveNamedThreads(%x)"), oldValue, newValue, FTaskTagScope::GetCurrentTag());
+			ensureMsgf(oldValue == newValue, TEXT("oldValue(%i) newValue(%i) If this check fails make sure that there is a FTaskTagScope(ETaskTag::EGameThread) as deep as possible on the current callstack, you can see the current value in ActiveNamedThreads(%x)"), oldValue, newValue, int32{EnumToUnderlyingType(FTaskTagScope::GetCurrentTag())});
 			newValue = oldValue;
 		}
 #endif
@@ -233,7 +233,7 @@ PRAGMA_DISABLE_DEPRECATION_WARNINGS
 	if (!UE::Tasks::Private::IsThreadRetractingTask())
 	{
 		bool oldValue = GSlateLoadingThreadId != 0 && FPlatformTLS::GetCurrentThreadId() == GSlateLoadingThreadId;
-		ensureMsgf(oldValue == newValue, TEXT("oldValue(%i) newValue(%i) If this check fails make sure that there is a FTaskTagScope(ETaskTag::ESlateThread) as as deep as possible on the current callstack, you can see the current value in ActiveNamedThreads(%x)"), oldValue, newValue, FTaskTagScope::GetCurrentTag());
+		ensureMsgf(oldValue == newValue, TEXT("oldValue(%i) newValue(%i) If this check fails make sure that there is a FTaskTagScope(ETaskTag::ESlateThread) as as deep as possible on the current callstack, you can see the current value in ActiveNamedThreads(%x)"), oldValue, newValue, int32{EnumToUnderlyingType(FTaskTagScope::GetCurrentTag())});
 		newValue = oldValue;
 	}
 #endif
@@ -265,7 +265,7 @@ PRAGMA_DISABLE_DEPRECATION_WARNINGS
 #if !UE_BUILD_SHIPPING && !UE_BUILD_TEST
 	{
 		bool oldValue = FPlatformTLS::GetCurrentThreadId() == GRenderThreadId;
-		ensureMsgf(oldValue == newValue, TEXT("oldValue(%i) newValue(%i) If this check fails make sure that there is a FTaskTagScope(ETaskTag::ERenderingThread) as deep as possible on the current callstack, you can see the current value in ActiveNamedThreads(%x)"), oldValue, newValue, FTaskTagScope::GetCurrentTag());
+		ensureMsgf(oldValue == newValue, TEXT("oldValue(%i) newValue(%i) If this check fails make sure that there is a FTaskTagScope(ETaskTag::ERenderingThread) as deep as possible on the current callstack, you can see the current value in ActiveNamedThreads(%x)"), oldValue, newValue, int32{EnumToUnderlyingType(FTaskTagScope::GetCurrentTag())});
 		newValue = oldValue;
 	}
 #endif
@@ -285,7 +285,7 @@ PRAGMA_DISABLE_DEPRECATION_WARNINGS
 	{
 		const uint32 CurrentThreadId = FPlatformTLS::GetCurrentThreadId();
 		bool oldValue = (GRenderThreadId == 0) ? (CurrentThreadId == GGameThreadId) || FTaskTagScope::IsRunningDuringStaticInit() : (CurrentThreadId == GRenderThreadId);
-		ensureMsgf(oldValue == newValue, TEXT("oldValue(%i) newValue(%i) If this check fails make sure that there is a FTaskTagScope(ETaskTag::ERenderingThread) as deep as possible on the current callstack, you can see the current value in ActiveNamedThreads(%x)"), oldValue, newValue, FTaskTagScope::GetCurrentTag());
+		ensureMsgf(oldValue == newValue, TEXT("oldValue(%i) newValue(%i) If this check fails make sure that there is a FTaskTagScope(ETaskTag::ERenderingThread) as deep as possible on the current callstack, you can see the current value in ActiveNamedThreads(%x)"), oldValue, newValue, int32{EnumToUnderlyingType(FTaskTagScope::GetCurrentTag())});
 		newValue = oldValue;
 	}
 #endif
@@ -322,7 +322,7 @@ PRAGMA_DISABLE_DEPRECATION_WARNINGS
 	{
 		const uint32 CurrentThreadId = FPlatformTLS::GetCurrentThreadId();
 		bool oldValue = (GRenderThreadId == 0) ?  true : CurrentThreadId != GGameThreadId;
-		ensureMsgf(oldValue == newValue, TEXT("oldValue(%i) newValue(%i) If this check fails make sure that there is a FTaskTagScope(ETaskTag::EParallelRenderingThread) as deep as possible on the current callstack, you can see the current value in ActiveNamedThreads(%x)"), oldValue, newValue, FTaskTagScope::GetCurrentTag());
+		ensureMsgf(oldValue == newValue, TEXT("oldValue(%i) newValue(%i) If this check fails make sure that there is a FTaskTagScope(ETaskTag::EParallelRenderingThread) as deep as possible on the current callstack, you can see the current value in ActiveNamedThreads(%x)"), oldValue, newValue, int32{EnumToUnderlyingType(FTaskTagScope::GetCurrentTag())});
 		newValue = oldValue;
 	}
 #endif
@@ -849,20 +849,8 @@ FEventRef::~FEventRef()
 -----------------------------------------------------------------------------*/
 
 FSharedEventRef::FSharedEventRef(EEventMode Mode  /* = EEventMode::AutoReset */)
-	: Ptr(Mode == EEventMode::ManualReset
-			? TLazySingleton<TEventPool<EEventMode::ManualReset>>::Get().GetRawEvent()
-			: TLazySingleton<TEventPool<EEventMode::AutoReset>>::Get().GetRawEvent(),
-		[](FEvent* Event)
-		{
-			if (Event->IsManualReset())
-			{
-				TLazySingleton<TEventPool<EEventMode::ManualReset>>::Get().ReturnRawEvent(Event);
-			}
-			else
-			{
-				TLazySingleton<TEventPool<EEventMode::AutoReset>>::Get().ReturnRawEvent(Event);
-			}
-		})
+	: Ptr(TLazySingleton<TEventPool<EEventMode::AutoReset>>::Get().GetRawEvent(),
+		[](FEvent* Event) { TLazySingleton<TEventPool<EEventMode::AutoReset>>::Get().ReturnRawEvent(Event); })
 {
 }
 
@@ -1587,7 +1575,19 @@ FTlsAutoCleanup* FThreadSingletonInitializer::Inject(FTlsAutoCleanup* Instance,
 
 void FTlsAutoCleanup::Register()
 {
-	static thread_local TArray<TUniquePtr<FTlsAutoCleanup>> TlsInstances;
+	// Destruction of TLS values need to occur in the reverse order of initialization, as thread_local does.
+	struct FTlsContainer : public TArray<TUniquePtr<FTlsAutoCleanup>>
+	{
+		~FTlsContainer()
+		{
+			for (int32 Index = Num() - 1; Index >= 0; Index--)
+			{
+				(*this)[Index].Reset();
+			}
+		}
+	};
+
+	static thread_local FTlsContainer TlsInstances;
 	TlsInstances.Add(TUniquePtr<FTlsAutoCleanup>(this));
 }
 
diff --git a/Engine/Source/Runtime/Core/Tests/HAL/SharedEventRefTest.cpp b/Engine/Source/Runtime/Core/Tests/HAL/SharedEventRefTest.cpp
deleted file mode 100644
