Added runtime network status monitoring to Windows platform

#Synced-CL 52272541

diff --git a/Engine/Source/Runtime/Core/Private/Windows/WindowsPlatformMisc.cpp b/Engine/Source/Runtime/Core/Private/Windows/WindowsPlatformMisc.cpp
--- a/Engine/Source/Runtime/Core/Private/Windows/WindowsPlatformMisc.cpp
+++ b/Engine/Source/Runtime/Core/Private/Windows/WindowsPlatformMisc.cpp
@@ -44,6 +44,7 @@
 #include "GenericPlatform/GenericPlatformDriver.h"
 #include "HAL/ThreadHeartBeat.h"
 #include "ProfilingDebugging/ExternalProfiler.h"
+#include "Async/TaskGraphInterfaces.h"
 
 // Resource includes.
 #include "Runtime/Launch/Resources/Windows/resource.h"
@@ -60,6 +61,7 @@ THIRD_PARTY_INCLUDES_START
 	#include <shlwapi.h>
 	#include <IPHlpApi.h>
 	#include <VersionHelpers.h>
+	#include <netlistmgr.h>
 #include "Windows/AllowWindowsPlatformAtomics.h"
 	#include <comdef.h>
 	#include <Wbemidl.h>
@@ -1091,10 +1093,14 @@ void FWindowsPlatformMisc::PlatformInit()
 
 	// Collecting storage information using dcom services.
 	CollectStorageInformation();
+
+	// Start listening for network connectivity changes via NLM events.
+	StartNetworkConnectionMonitoring();
 }
 
 void FWindowsPlatformMisc::PlatformTearDown()
 {
+	StopNetworkConnectionMonitoring();
 	FPlatformProcess::CeaseBeingFirstInstance();
 }
 
@@ -4547,4 +4553,326 @@ const class TMap<FString, FString>& FWindowsPlatformMisc::GetConfigRuleVars()
 #endif // WITH_APPLICATION_CORE && WITH_WINDOWS_CONFIGRULES
 }
 
+static INetworkListManager* GNetworkListManager = nullptr;
+static bool GNetworkMonitoringInitializedCOM = false;
+
+/** Cached network connection type, updated by NLM events and initial query.
+  * Under the current STA apartment model, both the ConnectivityChanged callback and all reads are serialized on the game thread via window message marshaling. */
+static ENetworkConnectionType GCachedNetworkConnectionType = ENetworkConnectionType::Unknown;
+
+/**
+ * Queries NLM network connections to determine if the active connection is WiFi or wired Ethernet
+ */
+static ENetworkConnectionType DetermineConnectionTypeFromNLM(INetworkListManager* NetworkListManager)
+{
+	// Get an enumerator over all NLM network connections
+	IEnumNetworkConnections* EnumConnections = nullptr;
+	HRESULT Result = NetworkListManager->GetNetworkConnections(&EnumConnections);
+	if (FAILED(Result) || EnumConnections == nullptr)
+	{
+		return ENetworkConnectionType::Unknown;
+	}
+
+	bool bHasWiFi = false;
+	bool bHasEthernet = false;
+
+	// Iterate each connection and check only those that are currently connected
+	INetworkConnection* Connection = nullptr;
+	while (EnumConnections->Next(1, &Connection, nullptr) == S_OK)
+	{
+		VARIANT_BOOL IsConnected = VARIANT_FALSE;
+		if (SUCCEEDED(Connection->get_IsConnected(&IsConnected)) && IsConnected != VARIANT_FALSE)
+		{
+			// Get the adapter GUID that backs this NLM connection
+			GUID AdapterId;
+			if (SUCCEEDED(Connection->GetAdapterId(&AdapterId)))
+			{
+				// Convert adapter GUID -> LUID -> interface index.
+				// We need the index because MIB_IFROW (used by GetIfEntry) is keyed by dwIndex
+				
+				NET_LUID Luid;
+				NET_IFINDEX IfIndex;
+
+				if (ConvertInterfaceGuidToLuid(&AdapterId, &Luid) == NO_ERROR &&
+					ConvertInterfaceLuidToIndex(&Luid, &IfIndex) == NO_ERROR)
+				{
+					// Query the interface row to read the physical adapter type
+					
+					MIB_IFROW IfRow = {};
+					IfRow.dwIndex = IfIndex;
+
+					if (GetIfEntry(&IfRow) == NO_ERROR)
+					{
+						if (IfRow.dwType == IF_TYPE_IEEE80211)
+						{
+							bHasWiFi = true;
+						}
+						else if (IfRow.dwType == IF_TYPE_ETHERNET_CSMACD)
+						{
+							bHasEthernet = true;
+						}
+					}
+				}
+			}
+		}
+
+		Connection->Release();
+	}
+
+	EnumConnections->Release();
+
+	// When both WiFi and Ethernet are active (e.g. docked laptop), prefer Ethernet, as the wired connection is typically the primary route
+	
+	if (bHasEthernet)
+	{
+		return ENetworkConnectionType::Ethernet;
+	}
+
+	if (bHasWiFi)
+	{
+		return ENetworkConnectionType::WiFi;
+	}
+
+	// Fallback: no recognized adapter type was found (e.g. VPN or cellular)
+	return ENetworkConnectionType::Unknown;
+}
+
+/**
+ * Queries NLM for the current network connection type. Expensive as it involves COM calls and interface enumeration.
+ * Should only be called during init and from NLM event callbacks, not per-frame.
+ * The result is stored in GCachedNetworkConnectionType.
+ */
+static ENetworkConnectionType QueryNetworkConnectionTypeFromNLM()
+{
+	if (GNetworkListManager == nullptr)
+	{
+		return ENetworkConnectionType::Unknown;
+	}
+
+	NLM_CONNECTIVITY Connectivity;
+	HRESULT GetConnectivityResult = GNetworkListManager->GetConnectivity(&Connectivity);
+	if (SUCCEEDED(GetConnectivityResult))
+	{
+		if (Connectivity & (NLM_CONNECTIVITY_IPV4_INTERNET | NLM_CONNECTIVITY_IPV6_INTERNET |
+			NLM_CONNECTIVITY_IPV4_LOCALNETWORK | NLM_CONNECTIVITY_IPV6_LOCALNETWORK))
+		{
+			return DetermineConnectionTypeFromNLM(GNetworkListManager);
+		}
+		else
+		{
+			return ENetworkConnectionType::None;
+		}
+	}
+
+	return ENetworkConnectionType::Unknown;
+}
+
+/**
+ * COM event listener that receives network connectivity change notifications from NLM.
+ * Registered in StartNetworkConnectionMonitoring, unregistered in StopNetworkConnectionMonitoring.
+ */
+class FNLMEventListener : public INetworkListManagerEvents
+{
+public:
+	FNLMEventListener() : RefCount(1), Cookie(0), ConnectionPoint(nullptr) {}
+	virtual ~FNLMEventListener() = default;
+
+	// IUnknown interface
+
+	HRESULT STDMETHODCALLTYPE QueryInterface(REFIID riid, void** ppvObject) override
+	{
+		if (riid == IID_IUnknown || riid == IID_INetworkListManagerEvents)
+		{
+			*ppvObject = this;
+			AddRef();
+			return S_OK;
+		}
+		*ppvObject = nullptr;
+		return E_NOINTERFACE;
+	}
+
+	ULONG STDMETHODCALLTYPE AddRef() override
+	{
+		// We use the PlatformAtomics methods to avoid having to add this to AllowWindowsPlatformTypes
+		return FPlatformAtomics::InterlockedIncrement(&RefCount);
+	}
+
+	ULONG STDMETHODCALLTYPE Release() override
+	{
+		// We use the PlatformAtomics methods to avoid having to add this to AllowWindowsPlatformTypes
+		int32 Ref = FPlatformAtomics::InterlockedDecrement(&RefCount);
+		if (Ref == 0)
+		{
+			delete this;
+		}
+		return Ref;
+	}
+
+	// ~ IUnknown interface
+
+	// INetworkListManagerEvents
+
+	HRESULT STDMETHODCALLTYPE ConnectivityChanged(NLM_CONNECTIVITY /*NewConnectivity*/) override
+	{
+		// We ignore the raw NLM_CONNECTIVITY bitmask and re-query NLM to also resolve adapter types.
+		const ENetworkConnectionType ConnectionType = QueryNetworkConnectionTypeFromNLM();
+		const ENetworkConnectionType PreviousType = GCachedNetworkConnectionType;
+		GCachedNetworkConnectionType = ConnectionType;
+
+		// NLM can fire ConnectivityChanged for events that don't affect the resolved type (e.g. a VPN connecting while Ethernet stays primary). Skip redundant broadcasts.
+		if (ConnectionType == PreviousType)
+		{
+			return S_OK;
+		}
+
+		UE_LOG(LogInit, Verbose, TEXT("Network connectivity changed (ConnectionType=%d)"), static_cast<int32>(ConnectionType));
+
+		if (FTaskGraphInterface::IsRunning())
+		{
+			FFunctionGraphTask::CreateAndDispatchWhenReady([ConnectionType]()
+			{
+				FCoreDelegates::OnNetworkConnectionChanged.Broadcast(ConnectionType);
+				FPlatformMisc::SetNetworkConnectionStatus(FPlatformMisc::NetworkConnectionTypeToStatus(ConnectionType));
+			}, TStatId(), nullptr, ENamedThreads::GameThread);
+		}
+		else
+		{
+			// Task graph not running yet (early startup). Set the status synchronously.
+			FPlatformMisc::SetNetworkConnectionStatus(FPlatformMisc::NetworkConnectionTypeToStatus(ConnectionType));
+		}
+		return S_OK;
+	}
+
+	// ~ INetworkListManagerEvents
+
+	/** Subscribe to NLM connectivity events. Returns true on success. */
+	bool Advise(INetworkListManager* NetworkListManager)
+	{
+		IConnectionPointContainer* Container = nullptr;
+		HRESULT Result = NetworkListManager->QueryInterface(IID_IConnectionPointContainer, (void**)&Container);
+		if (FAILED(Result) || Container == nullptr)
+		{
+			return false;
+		}
+
+		Result = Container->FindConnectionPoint(IID_INetworkListManagerEvents, &ConnectionPoint);
+		Container->Release();
+		if (FAILED(Result) || ConnectionPoint == nullptr)
+		{
+			return false;
+		}
+
+		Result = ConnectionPoint->Advise(this, &Cookie);
+		if (FAILED(Result))
+		{
+			ConnectionPoint->Release();
+			ConnectionPoint = nullptr;
+			return false;
+		}
+
+		return true;
+	}
+
+	/** Unsubscribe from NLM connectivity events. */
+	void Unadvise()
+	{
+		if (ConnectionPoint != nullptr)
+		{
+			if (Cookie != 0)
+			{
+				ConnectionPoint->Unadvise(Cookie);
+				Cookie = 0;
+			}
+			ConnectionPoint->Release();
+			ConnectionPoint = nullptr;
+		}
+	}
+
+private:
+	// Used by the IUnknown interface, type required by InterlockedIncrement/Decrement methods
+	volatile int32 RefCount;
+
+	// Token returned by Advise which identifies our registration
+	DWORD Cookie;
+
+	// Intermediary between the event source (INetworkListManager) and this listener
+	IConnectionPoint* ConnectionPoint;
+};
+
+static FNLMEventListener* GNLMEventListener = nullptr;
+
+ENetworkConnectionType FWindowsPlatformMisc::GetNetworkConnectionType()
+{
+	return GCachedNetworkConnectionType;
+}
+
+void FWindowsPlatformMisc::StartNetworkConnectionMonitoring()
+{
+	check(GNetworkListManager == nullptr && GNLMEventListener == nullptr);
+
+	// COM must be initialized on this thread before we can create COM objects.
+	// This defaults to a single-threaded apartment (COINIT_APARTMENTTHREADED), which means NLM callbacks are marshaled to this thread via window messages.
+	// Callbacks will only fire when the game thread pumps messages (FWindowsPlatformApplicationMisc::PumpMessages), so there may be latency under heavy game-thread load.
+	// This also guarantees that callbacks and teardown are serialized on the same thread, avoiding races.
+	if (!FWindowsPlatformMisc::CoInitialize())
+	{
+		UE_LOG(LogInit, Warning, TEXT("Failed to initialize COM for network monitoring"));
+		return;
+	}
+	
+	GNetworkMonitoringInitializedCOM = true;
+
+	HRESULT CoCreateInstanceResult = CoCreateInstance(CLSID_NetworkListManager, nullptr, CLSCTX_ALL, IID_INetworkListManager, (void**)&GNetworkListManager);
+	if (FAILED(CoCreateInstanceResult) || GNetworkListManager == nullptr)
+	{
+		UE_LOG(LogInit, Warning, TEXT("Failed to create INetworkListManager for network monitoring (HRESULT=0x%08X)"), CoCreateInstanceResult);
+		GNetworkListManager = nullptr;
+		FWindowsPlatformMisc::CoUninitialize();
+		GNetworkMonitoringInitializedCOM = false;
+		return;
+	}
+
+	GNLMEventListener = new FNLMEventListener();
+	if (!GNLMEventListener->Advise(GNetworkListManager))
+	{
+		UE_LOG(LogInit, Warning, TEXT("Failed to subscribe to INetworkListManagerEvents"));
+		GNLMEventListener->Release();
+		GNLMEventListener = nullptr;
+		GNetworkListManager->Release();
+		GNetworkListManager = nullptr;
+		FWindowsPlatformMisc::CoUninitialize();
+		GNetworkMonitoringInitializedCOM = false;
+		return;
+	}
+
+	// Set the initial network connection status.
+	const ENetworkConnectionType InitialConnectionType = QueryNetworkConnectionTypeFromNLM();
+	GCachedNetworkConnectionType = InitialConnectionType;
+	SetNetworkConnectionStatus(NetworkConnectionTypeToStatus(InitialConnectionType));
+}
+
+void FWindowsPlatformMisc::StopNetworkConnectionMonitoring()
+{
+	// Safe to access GNetworkListManager without a lock here: under the single-threaded COM apartment (COINIT_APARTMENTTHREADED), both the ConnectivityChanged callback and this teardown run on the same thread (the game thread), so they are serialized.
+	// If the apartment model is ever changed to multithreaded (COINIT_MULTITHREADED), a mutex or ref-counting guard would be needed.
+	if (GNLMEventListener != nullptr)
+	{
+		GNLMEventListener->Unadvise();
+		GNLMEventListener->Release();
+		GNLMEventListener = nullptr;
+	}
+
+	if (GNetworkListManager != nullptr)
+	{
+		GNetworkListManager->Release();
+		GNetworkListManager = nullptr;
+	}
+
+	if (GNetworkMonitoringInitializedCOM)
+	{
+		FWindowsPlatformMisc::CoUninitialize();
+		GNetworkMonitoringInitializedCOM = false;
+	}
+}
+
 #undef UE_CPUID
diff --git a/Engine/Source/Runtime/Core/Public/Windows/WindowsPlatformMisc.h b/Engine/Source/Runtime/Core/Public/Windows/WindowsPlatformMisc.h
--- a/Engine/Source/Runtime/Core/Public/Windows/WindowsPlatformMisc.h
+++ b/Engine/Source/Runtime/Core/Public/Windows/WindowsPlatformMisc.h
@@ -325,6 +325,23 @@ struct FWindowsPlatformMisc
 		return true;
 	}
 
+	/**
+	 * Returns the current network connection type using NLM (Network List Manager).
+	 */
+	static CORE_API ENetworkConnectionType GetNetworkConnectionType();
+
+	/**
+	 * Starts event-driven network connection monitoring via INetworkListManagerEvents.
+	 * Called from PlatformInit(). Updates network connection status on connectivity changes.
+	 */
+	static CORE_API void StartNetworkConnectionMonitoring();
+
+	/**
+	 * Stops network connection monitoring and releases NLM COM resources.
+	 * Called from PlatformTearDown().
+	 */
+	static CORE_API void StopNetworkConnectionMonitoring();
+
 	/**
 	 * Returns whether the device is considered a gaming handheld device. It's not a dynamic property, the status is cached at startup.
 	 */
