[AutoRTFM] Simplify the transactionally-safe MpscQueue.

This code was originally taken directly from the SpscQueue, replacing the atomic ops with a mutex. There was obvious room for simplification. I pointed Claude at the file and asked it to simplify the logic, and it immediately cut out the copied tail pointer, which is redundant when using a mutex.

(I asked for more simplification beyond that and it seemed like a wash, so I left it here.)

#Synced-CL 52137917

diff --git a/Engine/Source/Runtime/Core/Public/Containers/TransactionallySafeMpscQueue.h b/Engine/Source/Runtime/Core/Public/Containers/TransactionallySafeMpscQueue.h
--- a/Engine/Source/Runtime/Core/Public/Containers/TransactionallySafeMpscQueue.h
+++ b/Engine/Source/Runtime/Core/Public/Containers/TransactionallySafeMpscQueue.h
@@ -33,7 +33,7 @@ public:
 	[[nodiscard]] TTransactionallySafeMpscQueue()
 	{
 		FNode* Node = ::new (AllocatorType::Malloc(sizeof(FNode), alignof(FNode))) FNode;
-		Head = First = Tail = TailCopy = Node;
+		Head = First = Tail = Node;
 	}
 
 	~TTransactionallySafeMpscQueue()
@@ -41,18 +41,18 @@ public:
 		// Nobody should have a reference to this class anymore, but we still take the mutex to
 		// guarantee that the queue isn't modified by another thread while a transaction is underway.
 		UE::TScopeLock Lock(Mutex);
+		// Free all recycled/sentinel nodes (First up to and including Tail).
 		FNode* Node = First;
-
-		// Delete all sentinel or unoccupied nodes.
-		bool bContinue = false;
-		do
+		while (Node != Tail)
 		{
 			FNode* Next = Node->Next;
-			bContinue = Node != Tail;
 			AllocatorType::Free(Node);
 			Node = Next;
-		} 
-		while (bContinue);
+		}
+
+		// Free the sentinel (Tail) and advance to occupied nodes.
+		Node = Tail->Next;
+		AllocatorType::Free(Tail);
 
 		// Destroy and free all occupied nodes.
 		while (Node != nullptr)
@@ -143,47 +143,28 @@ private:
 		TTypeCompatibleBytes<ElementType> Value;
 	};
 
-private:
-	FNode* AllocFromCache()
-	{
-		// The mutex must _already be held_ when calling this function.
-		FNode* Node = First;
-		First = First->Next;
-		Node->Next = nullptr;
-		return Node;
-	}
-
 	FNode* AllocNode()
 	{
-		// First, we attempt to allocate a node from internal node cache.
+		// Attempt to recycle a consumed node from the free list (First..Tail range).
 		{
 			UE::TScopeLock Lock(Mutex);
-			if (First != TailCopy)
-			{
-				return AllocFromCache();
-			}
-
-			TailCopy = Tail;
-			if (First != TailCopy)
+			if (First != Tail)
 			{
-				return AllocFromCache();
+				FNode* Node = First;
+				First = First->Next;
+				Node->Next = nullptr;
+				return Node;
 			}
 		}
 
-		// Our cache is empty; allocate via ::operator new() instead.
+		// Free list is empty; allocate a new node.
 		return ::new (AllocatorType::Malloc(sizeof(FNode), alignof(FNode))) FNode();
 	}
 
-private:
 	// This mutex guards all accesses to this structure or to FNode::Next.
 	mutable UE::FTransactionallySafeMutex Mutex;
 
-	// consumer part 
-	// accessed mainly by consumer, infrequently by producer 
-	FNode* Tail; // tail of the queue 
-	// producer part 
-	// accessed only by producer 
-	FNode* Head; // head of the queue
-	FNode* First; // last unused node (tail of node cache) 
-	FNode* TailCopy; // points to the initially-created tail
+	FNode* Tail; // tail of the queue (consumer end)
+	FNode* Head; // head of the queue (producer end)
+	FNode* First; // head of the recycled node free list
 };
