Added UE::CStruct concept.

diff --git a/Engine/Source/Runtime/Core/Public/Concepts/UStruct.h b/Engine/Source/Runtime/Core/Public/Concepts/UStruct.h
new file mode 100644
--- /dev/null
+++ b/Engine/Source/Runtime/Core/Public/Concepts/UStruct.h
@@ -0,0 +1,26 @@
+// Copyright Epic Games, Inc. All Rights Reserved.
+
+#pragma once
+
+#include "Concepts/CompleteType.h"
+#include "Concepts/SameAs.h"
+
+#include <type_traits>
+
+class UScriptStruct;
+
+template <typename T>
+struct TBaseStructure;
+
+namespace UE
+{
+	/**
+	 * Concept which describes a UStruct.
+	 */
+	template <typename T>
+	concept CUStruct = UE::CCompleteType<T> &&
+		requires
+		{
+			{ TBaseStructure<std::remove_cv_t<T>>::Get() } -> UE::CSameAs<UScriptStruct*>;
+		};
+}
