Renamed UE::CUStruct to UE::CUScriptStruct.

diff --git a/Engine/Source/Runtime/Core/Public/Concepts/UScriptStruct.h b/Engine/Source/Runtime/Core/Public/Concepts/UScriptStruct.h
new file mode 100644
--- /dev/null
+++ b/Engine/Source/Runtime/Core/Public/Concepts/UScriptStruct.h
@@ -0,0 +1,26 @@
+// Copyright Epic Games, Inc. All Rights Reserved.
+
+#pragma once
+
+#include "Concepts/CompleteType.h"
+#include "Concepts/SameAs.h"
+
+#include <type_traits>
+
+class UScriptStruct;
+
+template <typename T>
+struct TBaseStructure;
+
+namespace UE
+{
+	/**
+	 * Concept which describes a UScriptStruct.
+	 */
+	template <typename T>
+	concept CUScriptStruct = UE::CCompleteType<T> &&
+		requires
+		{
+			{ TBaseStructure<std::remove_cv_t<T>>::Get() } -> UE::CSameAs<UScriptStruct*>;
+		};
+}
diff --git a/Engine/Source/Runtime/Core/Public/Concepts/UStruct.h b/Engine/Source/Runtime/Core/Public/Concepts/UStruct.h
deleted file mode 100644
diff --git a/Engine/Source/Runtime/Core/Public/UObject/NativeTypeToPropertyType.h b/Engine/Source/Runtime/Core/Public/UObject/NativeTypeToPropertyType.h
--- a/Engine/Source/Runtime/Core/Public/UObject/NativeTypeToPropertyType.h
+++ b/Engine/Source/Runtime/Core/Public/UObject/NativeTypeToPropertyType.h
@@ -4,7 +4,7 @@
 #include "CoreTypes.h"
 
 #include "Concepts/UObject.h"
-#include "Concepts/UStruct.h"
+#include "Concepts/UScriptStruct.h"
 #include "Concepts/DynamicDelegate.h"
 #include "Concepts/DynamicMulticastDelegate.h"
 #include "Concepts/SparseDynamicMulticastDelegate.h"
@@ -77,7 +77,7 @@ template <>                                         struct TNativeTypeToProperty
 template <typename T>                               struct TNativeTypeToPropertyType<TSoftObjectPtr<T>>                     { using Type = FSoftObjectProperty; };
 template <typename T>                               struct TNativeTypeToPropertyType<TSoftClassPtr<T>>                      { using Type = FSoftClassProperty; };
 template <>                                         struct TNativeTypeToPropertyType<FName>                                 { using Type = FNameProperty; };
-template <UE::CUStruct T>                           struct TNativeTypeToPropertyType<T>                                     { using Type = FStructProperty; };
+template <UE::CUScriptStruct T>                     struct TNativeTypeToPropertyType<T>                                     { using Type = FStructProperty; };
 template <>                                         struct TNativeTypeToPropertyType<FString>                               { using Type = FStrProperty; };
 template <>                                         struct TNativeTypeToPropertyType<FText>                                 { using Type = FTextProperty; };
 template <typename T>                               struct TNativeTypeToPropertyType<TArray<T, FDefaultAllocator>>          { using Type = FArrayProperty; };
