[Core] Initialize refcount to zero in TRefCountingMixin<T, ERefCountingMode::NotThreadSafe>.

Previously, the refcount would be left in an uninitialized state.

I didn't find any use of this specialization in the rest of the codebase; it might be a good idea to deprecate it.

#Synced-CL 52137917

diff --git a/Engine/Source/Runtime/Core/Public/Templates/RefCounting.h b/Engine/Source/Runtime/Core/Public/Templates/RefCounting.h
--- a/Engine/Source/Runtime/Core/Public/Templates/RefCounting.h
+++ b/Engine/Source/Runtime/Core/Public/Templates/RefCounting.h
@@ -279,7 +279,7 @@ private:
 		T::StaticDestroyObject(static_cast<const T*>(This));
 	}
 
-	mutable uint32 RefCount;
+	mutable uint32 RefCount = 0;
 };
 
 /**
