Clamp global boot verbosity overrides to each category's compile-time verbosity.

When a global verbosity is set via -LogCmds="global Warning", SetupSuppress applies it to every log category as they register. Previously, this would set DefaultVerbosity to the global value unconditionally, and ResetFromDefault would call SetVerbosity, which has an ensure that the new verbosity doesn't exceed CompileTimeVerbosity. Categories with a compile-time verbosity lower than the global override (e.g. Error-only categories hit by "global Warning") would trigger the ensure.

Fix: clamp the boot global override to Min(GlobalOverride, CompileTimeVerbosity) so that categories with a restrictive compile-time verbosity get clamped to be no greater than their compile-time verbosity.

#Synced-CL 52294780

#Synced-CL 52294780

diff --git a/Engine/Source/Runtime/Core/Private/Logging/LogSuppressionInterface.cpp b/Engine/Source/Runtime/Core/Private/Logging/LogSuppressionInterface.cpp
--- a/Engine/Source/Runtime/Core/Private/Logging/LogSuppressionInterface.cpp
+++ b/Engine/Source/Runtime/Core/Private/Logging/LogSuppressionInterface.cpp
@@ -415,7 +415,7 @@ class FLogSuppressionImplementation: public FLogSuppressionInterface, private FS
 			Boot = BootAssociations.Find(NAME_BootGlobal);
 			if (Boot)
 			{
-				Destination->DefaultVerbosity = *Boot;
+				Destination->DefaultVerbosity = FMath::Min<uint8>(*Boot, Destination->GetCompileTimeVerbosity());
 				Destination->ResetFromDefault();
 			}
 		}
