Set CrashReportingThread thread name on Windows to help debugging

#rnx

#Synced-CL 52288889

diff --git a/Engine/Source/Runtime/Core/Private/Windows/WindowsPlatformCrashContext.cpp b/Engine/Source/Runtime/Core/Private/Windows/WindowsPlatformCrashContext.cpp
--- a/Engine/Source/Runtime/Core/Private/Windows/WindowsPlatformCrashContext.cpp
+++ b/Engine/Source/Runtime/Core/Private/Windows/WindowsPlatformCrashContext.cpp
@@ -1295,6 +1295,8 @@ private:
 		__try
 #endif
 		{
+			FPlatformProcess::SetThreadName(TEXT("CrashReportingThread"));
+
 			bool bStopThreadRequested = false;
 			while (!bStopThreadRequested)
 			{
