change FCompression to explicit Init instead of init on first use

#Synced-CL 52315161

diff --git a/Engine/Source/Runtime/Core/Private/Compression/OodleDataCompression.cpp b/Engine/Source/Runtime/Core/Private/Compression/OodleDataCompression.cpp
--- a/Engine/Source/Runtime/Core/Private/Compression/OodleDataCompression.cpp
+++ b/Engine/Source/Runtime/Core/Private/Compression/OodleDataCompression.cpp
@@ -527,6 +527,8 @@ void CORE_API CompressionFormatInitOnFirstUseFromLock()
 	if ( GlobalOodleDataCompressionFormat != nullptr )
 		return;
 
+	check( IsInGameThread() );
+
 	GlobalOodleDataCompressionFormat = CreateOodleDataCompressionFormat();
 }
 
diff --git a/Engine/Source/Runtime/Core/Private/Misc/Compression.cpp b/Engine/Source/Runtime/Core/Private/Misc/Compression.cpp
--- a/Engine/Source/Runtime/Core/Private/Misc/Compression.cpp
+++ b/Engine/Source/Runtime/Core/Private/Misc/Compression.cpp
@@ -487,6 +487,16 @@ uint32 FCompression::GetCompressorVersion(FName FormatName)
 	return 0;
 }
 
+void FCompression::Init(void)
+{
+	FScopeLock Lock(&CompressionFormatsCriticalSection);
+	
+	// CompressionFormatInitOnFirstUseFromLock is safe to call multiple times
+	FOodleDataCompression::CompressionFormatInitOnFirstUseFromLock();
+
+	// OodleDataCompressionFormatInitOnFirstUseFromLock added it to the ModularFeatures list
+}
+
 ICompressionFormat* FCompression::GetCompressionFormat(FName FormatName, bool bErrorOnFailure)
 {
 	FScopeLock Lock(&CompressionFormatsCriticalSection);
@@ -496,11 +506,11 @@ ICompressionFormat* FCompression::GetCompressionFormat(FName FormatName, bool bE
 	{
 		if ( FormatName == NAME_Oodle )
 		{
-			// Oodle ICompressionFormat is created on first use :
+			// CompressionFormatInitOnFirstUseFromLock is a nop if already done
+			// this should now usually be done by FCompression::Init from LaunchEngineLoop
+			// but we need to keep this init-on-first-use here for the test harness
 			// inside CompressionFormatsCriticalSection lock
 			FOodleDataCompression::CompressionFormatInitOnFirstUseFromLock();
-
-			// OodleDataCompressionFormatInitOnFirstUseFromLock added it to the ModularFeatures list
 		}
 
 		TArray<ICompressionFormat*> Features = IModularFeatures::Get().GetModularFeatureImplementations<ICompressionFormat>(COMPRESSION_FORMAT_FEATURE_NAME);
diff --git a/Engine/Source/Runtime/Core/Public/Misc/Compression.h b/Engine/Source/Runtime/Core/Public/Misc/Compression.h
--- a/Engine/Source/Runtime/Core/Public/Misc/Compression.h
+++ b/Engine/Source/Runtime/Core/Public/Misc/Compression.h
@@ -36,6 +36,8 @@ struct FCompression
 	*	NAME_Oodle	Unused
 	*/
 
+	CORE_API static void Init(void);
+
 	/** Time spent compressing data in cycles. */
 	CORE_API static TAtomic<uint64> CompressorTimeCycles;
 	/** Number of bytes before compression.		*/
diff --git a/Engine/Source/Runtime/Launch/Private/LaunchEngineLoop.cpp b/Engine/Source/Runtime/Launch/Private/LaunchEngineLoop.cpp
--- a/Engine/Source/Runtime/Launch/Private/LaunchEngineLoop.cpp
+++ b/Engine/Source/Runtime/Launch/Private/LaunchEngineLoop.cpp
@@ -2733,8 +2733,14 @@ int32 FEngineLoop::PreInitPreStartupScreen(const TCHAR* CmdLine)
 	GLogConsole = GScopedLogConsole.Get();
 #endif
 
-	// init Oodle here
+	// init Oodle here ; note Decompressions are okay to run before this point
 	FOodleDataCompression::StartupPreInit();
+	
+#if WITH_ENGINE
+	// install named compression format modular features
+	// must be done after FCommandLine and DLL dirs are set up
+	FCompression::Init();
+#endif
 
 	{
 		SCOPED_BOOT_TIMING("LoadPreInitModules");
