Reduced severity of INetworkListManager initialization logs

#Synced-CL 52272541

diff --git a/Engine/Source/Runtime/Core/Private/Windows/WindowsPlatformMisc.cpp b/Engine/Source/Runtime/Core/Private/Windows/WindowsPlatformMisc.cpp
--- a/Engine/Source/Runtime/Core/Private/Windows/WindowsPlatformMisc.cpp
+++ b/Engine/Source/Runtime/Core/Private/Windows/WindowsPlatformMisc.cpp
@@ -4816,7 +4816,7 @@ void FWindowsPlatformMisc::StartNetworkConnectionMonitoring()
 	// This also guarantees that callbacks and teardown are serialized on the same thread, avoiding races.
 	if (!FWindowsPlatformMisc::CoInitialize())
 	{
-		UE_LOG(LogInit, Warning, TEXT("Failed to initialize COM for network monitoring"));
+		UE_LOG(LogInit, Log, TEXT("Failed to initialize COM for network monitoring"));
 		return;
 	}
 	
@@ -4825,7 +4825,7 @@ void FWindowsPlatformMisc::StartNetworkConnectionMonitoring()
 	HRESULT CoCreateInstanceResult = CoCreateInstance(CLSID_NetworkListManager, nullptr, CLSCTX_ALL, IID_INetworkListManager, (void**)&GNetworkListManager);
 	if (FAILED(CoCreateInstanceResult) || GNetworkListManager == nullptr)
 	{
-		UE_LOG(LogInit, Warning, TEXT("Failed to create INetworkListManager for network monitoring (HRESULT=0x%08X)"), CoCreateInstanceResult);
+		UE_LOG(LogInit, Log, TEXT("Failed to create INetworkListManager for network monitoring (HRESULT=0x%08X)"), CoCreateInstanceResult);
 		GNetworkListManager = nullptr;
 		FWindowsPlatformMisc::CoUninitialize();
 		GNetworkMonitoringInitializedCOM = false;
@@ -4835,7 +4835,7 @@ void FWindowsPlatformMisc::StartNetworkConnectionMonitoring()
 	GNLMEventListener = new FNLMEventListener();
 	if (!GNLMEventListener->Advise(GNetworkListManager))
 	{
-		UE_LOG(LogInit, Warning, TEXT("Failed to subscribe to INetworkListManagerEvents"));
+		UE_LOG(LogInit, Log, TEXT("Failed to subscribe to INetworkListManagerEvents"));
 		GNLMEventListener->Release();
 		GNLMEventListener = nullptr;
 		GNetworkListManager->Release();
