Fixed race condition that could occur when one UWorld is being created in an async-loading thread while another is being destroyed during GC on the game thread.

Changed module unload delegate to be thread safe variant.

diff --git a/Engine/Source/Runtime/Core/Public/Modules/ModuleManager.h b/Engine/Source/Runtime/Core/Public/Modules/ModuleManager.h
--- a/Engine/Source/Runtime/Core/Public/Modules/ModuleManager.h
+++ b/Engine/Source/Runtime/Core/Public/Modules/ModuleManager.h
@@ -563,7 +563,7 @@ public:
 	 *
 	 * @return The delegate.
 	 */
-	TMulticastDelegateRegistration<void()>& OnLibraryUnmount()
+	TTSMulticastDelegateRegistration<void()>& OnLibraryUnmount()
 	{
 		return OnLibraryUnmountCallback;
 	}
@@ -777,7 +777,7 @@ private:
 	TMulticastDelegate<void(TConstArrayView<FName>)> OnModulesUnloadCallback;
 
 	/** Multicast delegate called to remove anything left before a library unmounts. */
-	TMulticastDelegate<void()> OnLibraryUnmountCallback;
+	TTSMulticastDelegate<void()> OnLibraryUnmountCallback;
 
 	/** Delegate called to check whether any module to delay-unload still has live objects. */
 	TDelegate<bool(TConstArrayView<FName>)> CheckLiveObjectsInModulesCallback;
