Add missing tags for FName compression, including an FName/Compressed child tag for compressed fname data.

#Synced-CL 52288889

diff --git a/Engine/Source/Runtime/Core/Private/UObject/UnrealNames.cpp b/Engine/Source/Runtime/Core/Private/UObject/UnrealNames.cpp
--- a/Engine/Source/Runtime/Core/Private/UObject/UnrealNames.cpp
+++ b/Engine/Source/Runtime/Core/Private/UObject/UnrealNames.cpp
@@ -644,6 +644,7 @@ public:
 	template<class ScopeLock>
 	FNameEntryHandle Create(FNameStringView Name, TOptional<FNameEntryId> ComparisonId, FNameEntryHeader Header)
 	{
+		LLM_SCOPE(ELLMTag::FName);
 		FPlatformMisc::Prefetch(Blocks[CurrentBlock]);
 
 		FNameEntryHandle Handle = AllocateRegular<ScopeLock>(Name);
@@ -667,6 +668,7 @@ public:
 	template<class ScopeLock>
 	FNameEntryHandle CreateWithNumber(const FNameEntryId& StringPart, uint32 NumberPart, TOptional<FNameEntryId> ComparisonId, FNameEntryHeader Header)
 	{
+		LLM_SCOPE(ELLMTag::FName);
 		FPlatformMisc::Prefetch(Blocks[CurrentBlock]);
 		FNameEntryHandle Handle = AllocateNumbered<ScopeLock>();
 		FNameEntry& Entry = Resolve(Handle);
diff --git a/Engine/Source/Runtime/Core/Public/HAL/LowLevelMemTracker.h b/Engine/Source/Runtime/Core/Public/HAL/LowLevelMemTracker.h
--- a/Engine/Source/Runtime/Core/Public/HAL/LowLevelMemTracker.h
+++ b/Engine/Source/Runtime/Core/Public/HAL/LowLevelMemTracker.h
@@ -213,6 +213,7 @@ FString GetLLMProfilingDir();
 	macro(AudioSynthesis,						"AudioSynthesis",				GET_STATFNAME(STAT_AudioSynthesisLLM),						GET_STATFNAME(STAT_AudioSummaryLLM),			ELLMTag::Audio)\
 	macro(RealTimeCommunications,				"RealTimeCommunications",		GET_STATFNAME(STAT_RealTimeCommunicationsLLM),				NAME_None,										-1)\
 	macro(FName,								"FName",						GET_STATFNAME(STAT_FNameLLM),								GET_STATFNAME(STAT_EngineSummaryLLM),			-1)\
+	macro(FNameCompressed,						"FNameCompressed",				GET_STATFNAME(STAT_FNameCompressedLLM),						NAME_None,										ELLMTag::FName)\
 	macro(Networking,							"Networking",					GET_STATFNAME(STAT_NetworkingLLM),							GET_STATFNAME(STAT_NetworkingSummaryLLM),		-1)\
 	macro(Meshes,								"Meshes",						GET_STATFNAME(STAT_MeshesLLM),								GET_STATFNAME(STAT_MeshesSummaryLLM),			-1)\
 	macro(Stats,								"Stats",						GET_STATFNAME(STAT_StatsLLM),								GET_STATFNAME(STAT_EngineSummaryLLM),			-1)\
