[AutoRTFM] Remove AUTORTFM_INFER annotations that are now unnecessary.

The verse-clang compiler will implicitly apply AUTORTFM_INFER to all templated function and templated class methods when there are no other explicit annotations.
As such, these annotations are no longer needed.

#Synced-CL 52345207

diff --git a/Engine/Plugins/PCG/Source/PCG/Public/Helpers/PCGAsync.h b/Engine/Plugins/PCG/Source/PCG/Public/Helpers/PCGAsync.h
--- a/Engine/Plugins/PCG/Source/PCG/Public/Helpers/PCGAsync.h
+++ b/Engine/Plugins/PCG/Source/PCG/Public/Helpers/PCGAsync.h
@@ -4,7 +4,6 @@
 
 #include "Helpers/PCGAsyncState.h"
 
-#include "AutoRTFM/Defines.h"
 #include "Async/Async.h"
 #include "Containers/Array.h"
 #include "Containers/Map.h"
@@ -182,7 +181,7 @@ namespace FPCGAsync
 	* @returns true if the processing is done, false otherwise. Use this to know if you need to reschedule the task.
 	*/
 	template <typename InitializeFunc, typename ProcessElementFunc, typename MoveFunc, typename FinishedFunc>
-	AUTORTFM_INFER bool AsyncProcessingEx(FPCGAsyncState* AsyncState, int32 NumIterations, InitializeFunc&& Initialize, ProcessElementFunc&& ProcessElement, MoveFunc&& MoveData, FinishedFunc&& Finished, const bool bEnableTimeSlicing, const int32 ChunkSize = 64, const bool bAllowChunkSizeOverride = true)
+	bool AsyncProcessingEx(FPCGAsyncState* AsyncState, int32 NumIterations, InitializeFunc&& Initialize, ProcessElementFunc&& ProcessElement, MoveFunc&& MoveData, FinishedFunc&& Finished, const bool bEnableTimeSlicing, const int32 ChunkSize = 64, const bool bAllowChunkSizeOverride = true)
 	{
 		auto IterationInnerLoop = [Func = MoveTemp(ProcessElement)](int32 StartReadIndex, int32 StartWriteIndex, int32 Count) -> int32
 		{
@@ -234,7 +233,7 @@ namespace FPCGAsync
 	* @returns true if the processing is done, false otherwise. Use this to know if you need to reschedule the task.
 	*/
  	template <typename InitializeFunc, typename ProcessElementFunc>
-	AUTORTFM_INFER bool AsyncProcessingOneToOneEx(FPCGAsyncState* AsyncState, int32 NumIterations, InitializeFunc&& Initialize, ProcessElementFunc&& ProcessElement, const bool bEnableTimeSlicing, const int32 ChunkSize = 64, const bool bAllowChunkSizeOverride = true)
+	bool AsyncProcessingOneToOneEx(FPCGAsyncState* AsyncState, int32 NumIterations, InitializeFunc&& Initialize, ProcessElementFunc&& ProcessElement, const bool bEnableTimeSlicing, const int32 ChunkSize = 64, const bool bAllowChunkSizeOverride = true)
 	{
 		auto IterationInnerLoop = [Func = MoveTemp(ProcessElement)](int32 StartReadIndex, int32 StartWriteIndex, int32 Count) -> int32
 		{
@@ -288,7 +287,7 @@ namespace FPCGAsync
 	* @returns true if the processing is done, false otherwise. Use this to know if you need to reschedule the task.
 	*/
 	template <typename OutputType, typename Func, typename AllocatorType = FDefaultAllocator>
-	AUTORTFM_INFER bool AsyncProcessing(FPCGAsyncState* AsyncState, int32 NumIterations, TArray<OutputType, AllocatorType>& OutData, Func&& InFunc, const bool bEnableTimeSlicing, const int32 ChunkSize = 64, bool bAllowChunkSizeOverride = true)
+	bool AsyncProcessing(FPCGAsyncState* AsyncState, int32 NumIterations, TArray<OutputType, AllocatorType>& OutData, Func&& InFunc, const bool bEnableTimeSlicing, const int32 ChunkSize = 64, bool bAllowChunkSizeOverride = true)
 	{
 		auto Initialize = [&OutData, NumIterations]()
 		{
diff --git a/Engine/Plugins/PCG/Source/PCG/Public/SpatialAlgo/PCGOctreeQueries.h b/Engine/Plugins/PCG/Source/PCG/Public/SpatialAlgo/PCGOctreeQueries.h
--- a/Engine/Plugins/PCG/Source/PCG/Public/SpatialAlgo/PCGOctreeQueries.h
+++ b/Engine/Plugins/PCG/Source/PCG/Public/SpatialAlgo/PCGOctreeQueries.h
@@ -2,7 +2,6 @@
 
 #pragma once
 
-#include "AutoRTFM/Defines.h"
 #include "PCGPoint.h"
 #include "Data/PCGPointData.h"
 
@@ -23,7 +22,7 @@ private:
 public:
 	/** Query the internal octree to call a func on all points within some sphere. Callback takes the point (as const ref) and its squared distance to the center. */
 	template <typename Func = TFunctionRef<void(const UPCGBasePointData* PointData, int32 Index, double SquaredDistance)>>
-	AUTORTFM_INFER static PCG_API void ForEachPointInsideSphere(const UPCGBasePointData* InPointData, const FVector& InCenter, const double InRadius, Func InCallback);
+	static PCG_API void ForEachPointInsideSphere(const UPCGBasePointData* InPointData, const FVector& InCenter, const double InRadius, Func InCallback);
 
 	/** Get the closest point to a given position within the search distance, If bInDiscardCenter is true, will reject any points that are exactly at the center. Returns null if not found. */
 	static PCG_API const FPCGPoint* GetClosestPoint(const UPCGPointData* InPointData, const FVector& InCenter, const bool bInDiscardCenter, const double InSearchDistance = DefaultSearchDistance);
diff --git a/Engine/Source/Runtime/Core/Public/Async/Fundamental/Task.h b/Engine/Source/Runtime/Core/Public/Async/Fundamental/Task.h
--- a/Engine/Source/Runtime/Core/Public/Async/Fundamental/Task.h
+++ b/Engine/Source/Runtime/Core/Public/Async/Fundamental/Task.h
@@ -2,7 +2,6 @@
 
 #pragma once
 
-#include "AutoRTFM.h"
 #include "Logging/LogMacros.h"
 #include "Misc/EnumClassFlags.h"
 #include "TaskDelegate.h"
@@ -392,10 +391,10 @@ namespace LowLevelTasks
 
 		/** Initialize a new task and prepare for execution */
 		template<typename TRunnable>
-		AUTORTFM_INFER inline void Init(const TCHAR* InDebugName, ETaskPriority InPriority, TRunnable&& InRunnable, ETaskFlags Flags = ETaskFlags::DefaultFlags);
+		inline void Init(const TCHAR* InDebugName, ETaskPriority InPriority, TRunnable&& InRunnable, ETaskFlags Flags = ETaskFlags::DefaultFlags);
 
 		template<typename TRunnable>
-		AUTORTFM_INFER inline void Init(const TCHAR* InDebugName, TRunnable&& InRunnable, ETaskFlags Flags = ETaskFlags::DefaultFlags);
+		inline void Init(const TCHAR* InDebugName, TRunnable&& InRunnable, ETaskFlags Flags = ETaskFlags::DefaultFlags);
 
 		inline const TCHAR* GetDebugName() const;
 		inline ETaskPriority GetPriority() const;
diff --git a/Engine/Source/Runtime/Core/Public/Async/Fundamental/TaskDelegate.h b/Engine/Source/Runtime/Core/Public/Async/Fundamental/TaskDelegate.h
--- a/Engine/Source/Runtime/Core/Public/Async/Fundamental/TaskDelegate.h
+++ b/Engine/Source/Runtime/Core/Public/Async/Fundamental/TaskDelegate.h
@@ -131,7 +131,7 @@ namespace LowLevelTasks
 		struct TTaskDelegateImpl;
 
 		template<typename TCallableType>
-		struct AUTORTFM_INFER TTaskDelegateImpl<TCallableType, false> final : TTaskDelegateBase
+		struct TTaskDelegateImpl<TCallableType, false> final : TTaskDelegateBase
 		{
 			template<typename CallableT>
 			inline TTaskDelegateImpl(CallableT&& Callable, void* InlineData)
@@ -194,7 +194,7 @@ namespace LowLevelTasks
 		};
 
 		template<typename TCallableType>
-		struct AUTORTFM_INFER TTaskDelegateImpl<TCallableType, true> final : TTaskDelegateBase
+		struct TTaskDelegateImpl<TCallableType, true> final : TTaskDelegateBase
 		{
 		private:
 			inline TTaskDelegateImpl(void* DstData, void* SrcData)
diff --git a/Engine/Source/Runtime/Core/Public/Async/Future.h b/Engine/Source/Runtime/Core/Public/Async/Future.h
--- a/Engine/Source/Runtime/Core/Public/Async/Future.h
+++ b/Engine/Source/Runtime/Core/Public/Async/Future.h
@@ -2,7 +2,6 @@
 
 #pragma once
 
-#include "AutoRTFM/Defines.h"
 #include "CoreTypes.h"
 #include "Misc/AssertionMacros.h"
 #include "Templates/UnrealTemplate.h"
@@ -339,7 +338,7 @@ protected:
 	 * @return nothing at the moment but could return another future to allow future chaining
 	 */
 	template<typename Func>
-	AUTORTFM_INFER auto Then(Func Continuation);
+	auto Then(Func Continuation);
 
 	/**
 	 * Convenience wrapper for Then that
@@ -349,7 +348,7 @@ protected:
 	 * @return nothing at the moment but could return another future to allow future chaining
 	 */
 	template<typename Func>
-	AUTORTFM_INFER auto Next(Func Continuation);
+	auto Next(Func Continuation);
 
 	/**
 	 * Reset the future.
diff --git a/Engine/Source/Runtime/Core/Public/Containers/Array.h b/Engine/Source/Runtime/Core/Public/Containers/Array.h
--- a/Engine/Source/Runtime/Core/Public/Containers/Array.h
+++ b/Engine/Source/Runtime/Core/Public/Containers/Array.h
@@ -2,7 +2,6 @@
 
 #pragma once
 
-#include "AutoRTFM/Defines.h"
 #include "AutoRTFM/ReturnFromOpenMode.h"
 #include "CoreTypes.h"
 #include "Misc/AssertionMacros.h"
@@ -765,7 +764,7 @@ namespace UE::Core::Private
  *
  **/
 template<typename InElementType, typename InAllocatorType>
-class AUTORTFM_INFER TArray
+class TArray
 {
 	template <typename OtherInElementType, typename OtherAllocator>
 	friend class TArray;
diff --git a/Engine/Source/Runtime/Core/Public/Containers/Map.h b/Engine/Source/Runtime/Core/Public/Containers/Map.h
--- a/Engine/Source/Runtime/Core/Public/Containers/Map.h
+++ b/Engine/Source/Runtime/Core/Public/Containers/Map.h
@@ -4,7 +4,6 @@
 
 #include "CoreTypes.h"
 
-#include "AutoRTFM/Defines.h"
 #include "Algo/Reverse.h"
 #include "Containers/ContainerElementTypeCompatibility.h"
 #include "Containers/Set.h"
@@ -56,7 +55,7 @@ public:
 
 /** An initializer type for keys that's passed to the pair set when adding a new key. */
 template <typename KeyInitType>
-class AUTORTFM_INFER TKeyInitializer
+class TKeyInitializer
 {
 public:
 	std::conditional_t<std::is_rvalue_reference_v<KeyInitType>, KeyInitType&, KeyInitType> Key;
@@ -76,7 +75,7 @@ public:
 
 /** Defines how the map's pairs are hashed. */
 template<typename KeyType, typename ValueType, bool bInAllowDuplicateKeys>
-struct AUTORTFM_INFER TDefaultMapKeyFuncs : BaseKeyFuncs<TPair<KeyType, ValueType>, KeyType, bInAllowDuplicateKeys>
+struct TDefaultMapKeyFuncs : BaseKeyFuncs<TPair<KeyType, ValueType>, KeyType, bInAllowDuplicateKeys>
 {
 	typedef typename TTypeTraits<KeyType>::ConstPointerType KeyInitType;
 	typedef const TPairInitializer<typename TTypeTraits<KeyType>::ConstInitType, typename TTypeTraits<ValueType>::ConstInitType>& ElementInitType;
diff --git a/Engine/Source/Runtime/Core/Public/Containers/Map.h.inl b/Engine/Source/Runtime/Core/Public/Containers/Map.h.inl
--- a/Engine/Source/Runtime/Core/Public/Containers/Map.h.inl
+++ b/Engine/Source/Runtime/Core/Public/Containers/Map.h.inl
@@ -25,7 +25,7 @@
  *    the cache misses of reading key data and doing the hashing *before* acquiring the lock.
  **/
 template <typename KeyType, typename ValueType, typename SetAllocator, typename KeyFuncs>
-class AUTORTFM_INFER TMAPBASE
+class TMAPBASE
 {
 	template <typename OtherKeyType, typename OtherValueType, typename OtherSetAllocator, typename OtherKeyFuncs>
 	friend class TMAPBASE;
@@ -1019,7 +1019,7 @@ public:
 
 /** The base type of sortable maps. */
 template <typename KeyType, typename ValueType, typename SetAllocator, typename KeyFuncs>
-class AUTORTFM_INFER TSORTABLEMAPBASE : public TMAPBASE<KeyType, ValueType, SetAllocator, KeyFuncs>
+class TSORTABLEMAPBASE : public TMAPBASE<KeyType, ValueType, SetAllocator, KeyFuncs>
 {
 protected:
 	typedef TMAPBASE<KeyType, ValueType, SetAllocator, KeyFuncs> Super;
@@ -1176,7 +1176,7 @@ private:
 
 /** A TMAPBASE specialization that only allows a single value associated with each key.*/
 template<typename InKeyType, typename InValueType, typename SetAllocator /*= FDefaultSetAllocator*/, typename KeyFuncs /*= TDefaultMapHashableKeyFuncs<KeyType,ValueType,false>*/>
-class AUTORTFM_INFER TMAP : public TSORTABLEMAPBASE<InKeyType, InValueType, SetAllocator, KeyFuncs>
+class TMAP : public TSORTABLEMAPBASE<InKeyType, InValueType, SetAllocator, KeyFuncs>
 {
 	template <typename, typename>
 	friend class TScriptMap;
diff --git a/Engine/Source/Runtime/Core/Public/Containers/SetUtilities.h b/Engine/Source/Runtime/Core/Public/Containers/SetUtilities.h
--- a/Engine/Source/Runtime/Core/Public/Containers/SetUtilities.h
+++ b/Engine/Source/Runtime/Core/Public/Containers/SetUtilities.h
@@ -2,7 +2,6 @@
 
 #pragma once
 
-#include "AutoRTFM/Defines.h"
 #include "CoreTypes.h"
 #include "Serialization/MemoryLayout.h"
 #include "Templates/MemoryOps.h"
@@ -38,7 +37,7 @@ struct BaseKeyFuncs
  * A default implementation of the KeyFuncs used by TSet which uses the element as a key.
  */
 template<typename ElementType,bool bInAllowDuplicateKeys /*= false*/>
-struct AUTORTFM_INFER DefaultKeyFuncs : BaseKeyFuncs<ElementType,ElementType,bInAllowDuplicateKeys>
+struct DefaultKeyFuncs : BaseKeyFuncs<ElementType,ElementType,bInAllowDuplicateKeys>
 {
 	typedef typename TTypeTraits<ElementType>::ConstPointerType KeyInitType;
 	typedef typename TCallTraits<ElementType>::ParamType ElementInitType;
@@ -85,7 +84,7 @@ struct AUTORTFM_INFER DefaultKeyFuncs : BaseKeyFuncs<ElementType,ElementType,bIn
 /** This is used to provide type specific behavior for a move which will destroy B. */
 /** Should be in UnrealTemplate but isn't for Clang build reasons - will move later */
 template<typename T>
-AUTORTFM_INFER inline void MoveByRelocate(T& A, T& B)
+inline void MoveByRelocate(T& A, T& B)
 {
 	// Destruct the previous value of A.
 	A.~T();
diff --git a/Engine/Source/Runtime/Core/Public/Containers/SparseArray.h b/Engine/Source/Runtime/Core/Public/Containers/SparseArray.h
--- a/Engine/Source/Runtime/Core/Public/Containers/SparseArray.h
+++ b/Engine/Source/Runtime/Core/Public/Containers/SparseArray.h
@@ -2,7 +2,6 @@
 
 #pragma once
 
-#include "AutoRTFM/Defines.h"
 #include "CoreTypes.h"
 #include "Misc/AssertionMacros.h"
 #include "HAL/UnrealMemory.h"
@@ -521,7 +520,7 @@ class TScriptSparseArray;
  *
  **/
 template<typename InElementType, typename Allocator /*= FDefaultSparseArrayAllocator */>
-class AUTORTFM_INFER TSparseArray : public TSparseArrayBase<sizeof(InElementType), alignof(InElementType), Allocator>
+class TSparseArray : public TSparseArrayBase<sizeof(InElementType), alignof(InElementType), Allocator>
 {
 	using SuperType = TSparseArrayBase<sizeof(InElementType), alignof(InElementType), Allocator>;
 	using ElementType = InElementType;
diff --git a/Engine/Source/Runtime/Core/Public/Containers/SparseSet.h.inl b/Engine/Source/Runtime/Core/Public/Containers/SparseSet.h.inl
--- a/Engine/Source/Runtime/Core/Public/Containers/SparseSet.h.inl
+++ b/Engine/Source/Runtime/Core/Public/Containers/SparseSet.h.inl
@@ -4,7 +4,6 @@
 #error "SparseSet.h.inl should only be included after defining UE_TSPARSE_SET"
 #endif
 
-#include "AutoRTFM/Defines.h"
 #include "Containers/ContainerAllocationPolicies.h"
 #include "Containers/ContainerElementTypeCompatibility.h"
 #include "Containers/SetUtilities.h"
@@ -56,7 +55,7 @@ template<
 	typename KeyFuncs /*= DefaultKeyFuncs<ElementType>*/,
 	typename Allocator /*= FDefaultSparseSetAllocator*/
 	>
-class AUTORTFM_INFER UE_TSPARSE_SET
+class UE_TSPARSE_SET
 {
 public:
 	typedef InElementType ElementType;
diff --git a/Engine/Source/Runtime/Core/Public/Containers/SparseSetElement.h b/Engine/Source/Runtime/Core/Public/Containers/SparseSetElement.h
--- a/Engine/Source/Runtime/Core/Public/Containers/SparseSetElement.h
+++ b/Engine/Source/Runtime/Core/Public/Containers/SparseSetElement.h
@@ -2,14 +2,13 @@
 
 #pragma once
 
-#include "AutoRTFM/Defines.h"
 #include "Containers/SetUtilities.h"
 #include "Serialization/StructuredArchive.h"
 #include <initializer_list>
 #include <type_traits>
 
 template<typename InElementType, bool bTypeLayout>
-class AUTORTFM_INFER TSparseSetElementBase
+class TSparseSetElementBase
 {
 public:
 	typedef InElementType ElementType;
@@ -42,7 +41,7 @@ public:
 };
 
 template<typename InElementType>
-class AUTORTFM_INFER TSparseSetElementBase<InElementType, true>
+class TSparseSetElementBase<InElementType, true>
 {
 	DECLARE_INLINE_TYPE_LAYOUT(TSparseSetElementBase, NonVirtual);
 public:
@@ -77,7 +76,7 @@ public:
 
 /** An element in the set. */
 template <typename InElementType>
-class AUTORTFM_INFER TSparseSetElement : public TSparseSetElementBase<InElementType, THasTypeLayout<InElementType>::Value>
+class TSparseSetElement : public TSparseSetElementBase<InElementType, THasTypeLayout<InElementType>::Value>
 {
 	using Super = TSparseSetElementBase<InElementType, THasTypeLayout<InElementType>::Value>;
 public:
diff --git a/Engine/Source/Runtime/Core/Public/Delegates/DelegateInstancesImpl.h b/Engine/Source/Runtime/Core/Public/Delegates/DelegateInstancesImpl.h
--- a/Engine/Source/Runtime/Core/Public/Delegates/DelegateInstancesImpl.h
+++ b/Engine/Source/Runtime/Core/Public/Delegates/DelegateInstancesImpl.h
@@ -10,7 +10,6 @@
 
 // HEADER_UNIT_SKIP - Not included directly
 
-#include "AutoRTFM.h"
 #include "CoreTypes.h"
 #include "Delegates/DelegateInstanceInterface.h"
 #include "Delegates/DelegateInstancesImplFwd.h"
@@ -57,7 +56,7 @@ protected:
  *****************************************************************************/
 
 template <class UserClass, typename RetValType, typename... ParamTypes, typename UserPolicy, typename... VarTypes>
-class AUTORTFM_INFER TBaseUFunctionDelegateInstance<UserClass, RetValType(ParamTypes...), UserPolicy, VarTypes...> : public TCommonDelegateInstanceState<RetValType(ParamTypes...), UserPolicy, VarTypes...>
+class TBaseUFunctionDelegateInstance<UserClass, RetValType(ParamTypes...), UserPolicy, VarTypes...> : public TCommonDelegateInstanceState<RetValType(ParamTypes...), UserPolicy, VarTypes...>
 {
 private:
 	using Super            = TCommonDelegateInstanceState<RetValType(ParamTypes...), UserPolicy, VarTypes...>;
@@ -183,7 +182,7 @@ public:
 
 
 template <bool bConst, class UserClass, ESPMode SPMode, typename RetValType, typename... ParamTypes, typename UserPolicy, typename... VarTypes>
-class AUTORTFM_INFER TBaseSPMethodDelegateInstance<bConst, UserClass, SPMode, RetValType(ParamTypes...), UserPolicy, VarTypes...> : public TCommonDelegateInstanceState<RetValType(ParamTypes...), UserPolicy, VarTypes...>
+class TBaseSPMethodDelegateInstance<bConst, UserClass, SPMode, RetValType(ParamTypes...), UserPolicy, VarTypes...> : public TCommonDelegateInstanceState<RetValType(ParamTypes...), UserPolicy, VarTypes...>
 {
 private:
 	using Super            = TCommonDelegateInstanceState<RetValType(ParamTypes...), UserPolicy, VarTypes...>;
@@ -300,7 +299,7 @@ protected:
 
 
 template <ESPMode SPMode, typename RetValType, typename... ParamTypes, typename UserPolicy, typename FunctorType, typename... VarTypes>
-class AUTORTFM_INFER TBaseSPLambdaDelegateInstance<SPMode, RetValType(ParamTypes...), UserPolicy, FunctorType, VarTypes...> : public TCommonDelegateInstanceState<RetValType(ParamTypes...), UserPolicy, VarTypes...>
+class TBaseSPLambdaDelegateInstance<SPMode, RetValType(ParamTypes...), UserPolicy, FunctorType, VarTypes...> : public TCommonDelegateInstanceState<RetValType(ParamTypes...), UserPolicy, VarTypes...>
 {
 private:
 	static_assert(std::is_same_v<FunctorType, std::remove_reference_t<FunctorType>>, "FunctorType cannot be a reference");
@@ -392,7 +391,7 @@ private:
 
 
 template <bool bConst, class UserClass, typename RetValType, typename... ParamTypes, typename UserPolicy, typename... VarTypes>
-class AUTORTFM_INFER TBaseRawMethodDelegateInstance<bConst, UserClass, RetValType(ParamTypes...), UserPolicy, VarTypes...> : public TCommonDelegateInstanceState<RetValType(ParamTypes...), UserPolicy, VarTypes...>
+class TBaseRawMethodDelegateInstance<bConst, UserClass, RetValType(ParamTypes...), UserPolicy, VarTypes...> : public TCommonDelegateInstanceState<RetValType(ParamTypes...), UserPolicy, VarTypes...>
 {
 private:
 	static_assert(!UE::Delegates::Private::IsUObjectPtr((UserClass*)nullptr), "You cannot use raw method delegates with UObjects.");
@@ -509,7 +508,7 @@ protected:
 
 
 template <bool bConst, class UserClass, typename RetValType, typename... ParamTypes, typename UserPolicy, typename... VarTypes>
-class AUTORTFM_INFER TBaseUObjectMethodDelegateInstance<bConst, UserClass, RetValType(ParamTypes...), UserPolicy, VarTypes...> : public TCommonDelegateInstanceState<RetValType(ParamTypes...), UserPolicy, VarTypes...>
+class TBaseUObjectMethodDelegateInstance<bConst, UserClass, RetValType(ParamTypes...), UserPolicy, VarTypes...> : public TCommonDelegateInstanceState<RetValType(ParamTypes...), UserPolicy, VarTypes...>
 {
 private:
 	using Super            = TCommonDelegateInstanceState<RetValType(ParamTypes...), UserPolicy, VarTypes...>;
@@ -641,7 +640,7 @@ protected:
 
 
 template <typename RetValType, typename... ParamTypes, typename UserPolicy, typename... VarTypes>
-class AUTORTFM_INFER TBaseStaticDelegateInstance<RetValType(ParamTypes...), UserPolicy, VarTypes...> : public TCommonDelegateInstanceState<RetValType(ParamTypes...), UserPolicy, VarTypes...>
+class TBaseStaticDelegateInstance<RetValType(ParamTypes...), UserPolicy, VarTypes...> : public TCommonDelegateInstanceState<RetValType(ParamTypes...), UserPolicy, VarTypes...>
 {
 private:
 	using Super            = TCommonDelegateInstanceState<RetValType(ParamTypes...), UserPolicy, VarTypes...>;
@@ -732,7 +731,7 @@ private:
 
 
 template <typename RetValType, typename... ParamTypes, typename UserPolicy, typename FunctorType, typename... VarTypes>
-class AUTORTFM_INFER TBaseFunctorDelegateInstance<RetValType(ParamTypes...), UserPolicy, FunctorType, VarTypes...> : public TCommonDelegateInstanceState<RetValType(ParamTypes...), UserPolicy, VarTypes...>
+class TBaseFunctorDelegateInstance<RetValType(ParamTypes...), UserPolicy, FunctorType, VarTypes...> : public TCommonDelegateInstanceState<RetValType(ParamTypes...), UserPolicy, VarTypes...>
 {
 private:
 	static_assert(std::is_same_v<FunctorType, std::remove_reference_t<FunctorType>>, "FunctorType cannot be a reference");
@@ -817,7 +816,7 @@ private:
 
 
 template <typename RetValType, typename... ParamTypes, typename UserPolicy, typename FunctorType, typename... VarTypes>
-class AUTORTFM_INFER TWeakBaseFunctorDelegateInstance<RetValType(ParamTypes...), UserPolicy, FunctorType, VarTypes...> : public TCommonDelegateInstanceState<RetValType(ParamTypes...), UserPolicy, VarTypes...>
+class TWeakBaseFunctorDelegateInstance<RetValType(ParamTypes...), UserPolicy, FunctorType, VarTypes...> : public TCommonDelegateInstanceState<RetValType(ParamTypes...), UserPolicy, VarTypes...>
 {
 private:
 	static_assert(std::is_same_v<FunctorType, std::remove_reference_t<FunctorType>>, "FunctorType cannot be a reference");
diff --git a/Engine/Source/Runtime/Core/Public/Experimental/Async/AwaitableTask.h b/Engine/Source/Runtime/Core/Public/Experimental/Async/AwaitableTask.h
--- a/Engine/Source/Runtime/Core/Public/Experimental/Async/AwaitableTask.h
+++ b/Engine/Source/Runtime/Core/Public/Experimental/Async/AwaitableTask.h
@@ -4,7 +4,6 @@
 #include "Async/Fundamental/Task.h"
 #include "Async/Fundamental/Scheduler.h"
 #include "Async/ManualResetEvent.h"
-#include "AutoRTFM.h"
 #include "Misc/Launder.h"
 
 template<typename>
@@ -148,7 +147,7 @@ namespace AwaitableTask_Detail
 
 	//the TPromiseVTable is used to type erase the PromiseType and therefore its lambda type
 	template<typename PromiseType>
-	class AUTORTFM_INFER TPromiseVTable final : public TPromiseVTableBase<typename PromiseType::ReturnType>
+	class TPromiseVTable final : public TPromiseVTableBase<typename PromiseType::ReturnType>
 	{
 		using ReturnType = typename PromiseType::ReturnType;
 
@@ -203,7 +202,7 @@ namespace AwaitableTask_Detail
 	};
 
 	template<typename TReturnType, typename CallableType>
-	class AUTORTFM_INFER TPromise final : public FPromiseBase
+	class TPromise final : public FPromiseBase
 	{
 	public:
 		using ReturnType = TReturnType;
@@ -282,7 +281,7 @@ namespace AwaitableTask_Detail
 	};
 
 	template<typename CallableType>
-	class AUTORTFM_INFER TPromise<void, CallableType> final : public FPromiseBase
+	class TPromise<void, CallableType> final : public FPromiseBase
 	{
 	public:
 		using ReturnType = void;
diff --git a/Engine/Source/Runtime/Core/Public/Math/GenericOctree.h b/Engine/Source/Runtime/Core/Public/Math/GenericOctree.h
--- a/Engine/Source/Runtime/Core/Public/Math/GenericOctree.h
+++ b/Engine/Source/Runtime/Core/Public/Math/GenericOctree.h
@@ -6,7 +6,6 @@
 
 #pragma once
 
-#include "AutoRTFM/Defines.h"
 #include "Containers/Array.h"
 #include "Containers/ArrayView.h"
 #include "Containers/ContainerAllocationPolicies.h"
@@ -640,7 +639,7 @@ private:
 	}
 
 	template<typename IterateFunc>
-	AUTORTFM_INFER void FindElementsWithBoundsTestInternal(FNodeIndex CurrentNodeIndex, const FOctreeNodeContext& NodeContext, const FBoxCenterAndExtent& BoxBounds, const IterateFunc& Func) const
+	void FindElementsWithBoundsTestInternal(FNodeIndex CurrentNodeIndex, const FOctreeNodeContext& NodeContext, const FBoxCenterAndExtent& BoxBounds, const IterateFunc& Func) const
 	{
 		if (TreeNodes[CurrentNodeIndex].InclusiveNumElements > 0)
 		{
diff --git a/Engine/Source/Runtime/Core/Public/Misc/Optional.h b/Engine/Source/Runtime/Core/Public/Misc/Optional.h
--- a/Engine/Source/Runtime/Core/Public/Misc/Optional.h
+++ b/Engine/Source/Runtime/Core/Public/Misc/Optional.h
@@ -2,7 +2,6 @@
 
 #pragma once
 
-#include "AutoRTFM/Defines.h"
 #include "AutoRTFM/ReturnFromOpenMode.h"
 #include "CoreTypes.h"
 #include "Misc/AssertionMacros.h"
@@ -45,7 +44,7 @@ namespace UE::Core::Private
  * Otherwise GetValue() is not meaningful.
  */
 template<typename OptionalType>
-struct AUTORTFM_INFER TOptional
+struct TOptional
 {
 private:
 	using FOptional = UE::Core::Private::FOptional;
diff --git a/Engine/Source/Runtime/Core/Public/Misc/ScopeExit.h b/Engine/Source/Runtime/Core/Public/Misc/ScopeExit.h
--- a/Engine/Source/Runtime/Core/Public/Misc/ScopeExit.h
+++ b/Engine/Source/Runtime/Core/Public/Misc/ScopeExit.h
@@ -2,8 +2,6 @@
 
 #pragma once
 
-#include "AutoRTFM/Defines.h"
-
 namespace ScopeExitSupport
 {
 	/**
@@ -12,7 +10,7 @@ namespace ScopeExitSupport
 	 * RAII class that calls a lambda when it is destroyed.
 	 */
 	template <typename FuncType>
-	class AUTORTFM_INFER TScopeGuard
+	class TScopeGuard
 	{
 		TScopeGuard(TScopeGuard&&) = delete;
 		TScopeGuard(const TScopeGuard&) = delete;
diff --git a/Engine/Source/Runtime/Core/Public/Tasks/TaskPrivate.h b/Engine/Source/Runtime/Core/Public/Tasks/TaskPrivate.h
--- a/Engine/Source/Runtime/Core/Public/Tasks/TaskPrivate.h
+++ b/Engine/Source/Runtime/Core/Public/Tasks/TaskPrivate.h
@@ -863,7 +863,7 @@ protected:
 
 		// a specialization for tasks that don't return results
 		template<typename TaskBodyType>
-		class AUTORTFM_INFER TExecutableTaskBase<TaskBodyType, typename TEnableIf<std::is_same_v<TInvokeResult_T<TaskBodyType>, void>>::Type> : public FTaskBase
+		class TExecutableTaskBase<TaskBodyType, typename TEnableIf<std::is_same_v<TInvokeResult_T<TaskBodyType>, void>>::Type> : public FTaskBase
 		{
 			UE_NONCOPYABLE(TExecutableTaskBase);
 
diff --git a/Engine/Source/Runtime/Core/Public/Templates/Function.h b/Engine/Source/Runtime/Core/Public/Templates/Function.h
--- a/Engine/Source/Runtime/Core/Public/Templates/Function.h
+++ b/Engine/Source/Runtime/Core/Public/Templates/Function.h
@@ -70,7 +70,7 @@ namespace UE::Core::Private::Function
 	};
 
 	template <typename T, bool bUnique, bool bOnHeap>
-	struct AUTORTFM_INFER TFunction_OwnedObject : public IFunction_OwnedObject
+	struct TFunction_OwnedObject : public IFunction_OwnedObject
 	{
 		template <typename ArgType>
 		explicit TFunction_OwnedObject(ArgType&& Arg)
@@ -200,7 +200,7 @@ namespace UE::Core::Private::Function
 	};
 
 	template <bool bUnique>
-	struct AUTORTFM_INFER TFunctionStorage : FFunctionStorage
+	struct TFunctionStorage : FFunctionStorage
 	{
 		TFunctionStorage() = default;
 
@@ -283,7 +283,7 @@ namespace UE::Core::Private::Function
 	 * A class which is used to instantiate the code needed to call a bound function.
 	 */
 	template <typename Functor, typename Ret, typename... ParamTypes>
-	struct AUTORTFM_INFER TFunctionRefCaller
+	struct TFunctionRefCaller
 	{
 		static Ret Call(void* Obj, ParamTypes&... Params)
 		{
@@ -312,7 +312,7 @@ namespace UE::Core::Private::Function
 	struct TFunctionRefBase;
 
 	template <typename StorageType, typename Ret, typename... ParamTypes>
-	struct AUTORTFM_INFER TFunctionRefBase<StorageType, Ret (ParamTypes...)>
+	struct TFunctionRefBase<StorageType, Ret (ParamTypes...)>
 	{
 		template <typename OtherStorageType, typename OtherFuncType>
 		friend struct TFunctionRefBase;
@@ -553,7 +553,7 @@ namespace UE::Core::Private::Function
  * }
  */
 template <typename Ret, typename... ParamTypes>
-class AUTORTFM_INFER TFunctionRef<Ret(ParamTypes...)> final : public UE::Core::Private::Function::TFunctionRefBase<UE::Core::Private::Function::FFunctionRefStoragePolicy, Ret(ParamTypes...)>
+class TFunctionRef<Ret(ParamTypes...)> final : public UE::Core::Private::Function::TFunctionRefBase<UE::Core::Private::Function::FFunctionRefStoragePolicy, Ret(ParamTypes...)>
 {
 	using Super = UE::Core::Private::Function::TFunctionRefBase<UE::Core::Private::Function::FFunctionRefStoragePolicy, Ret(ParamTypes...)>;
 
@@ -636,7 +636,7 @@ public:
  * }
  */
 template <typename Ret, typename... ParamTypes>
-class AUTORTFM_INFER TFunction<Ret(ParamTypes...)> final : public UE::Core::Private::Function::TFunctionRefBase<UE::Core::Private::Function::TFunctionStorage<false>, Ret(ParamTypes...)>
+class TFunction<Ret(ParamTypes...)> final : public UE::Core::Private::Function::TFunctionRefBase<UE::Core::Private::Function::TFunctionStorage<false>, Ret(ParamTypes...)>
 {
 	using Super = UE::Core::Private::Function::TFunctionRefBase<UE::Core::Private::Function::TFunctionStorage<false>, Ret(ParamTypes...)>;
 
@@ -751,7 +751,7 @@ public:
  * Foo(MoveTemp(MovableFunc)); // ok
  */
 template <typename Ret, typename... ParamTypes>
-class AUTORTFM_INFER TUniqueFunction<Ret(ParamTypes...)> final : public UE::Core::Private::Function::TFunctionRefBase<UE::Core::Private::Function::TFunctionStorage<true>, Ret(ParamTypes...)>
+class TUniqueFunction<Ret(ParamTypes...)> final : public UE::Core::Private::Function::TFunctionRefBase<UE::Core::Private::Function::TFunctionStorage<true>, Ret(ParamTypes...)>
 {
 	using Super = UE::Core::Private::Function::TFunctionRefBase<UE::Core::Private::Function::TFunctionStorage<true>, Ret(ParamTypes...)>;
 
diff --git a/Engine/Source/Runtime/Core/Public/Templates/FunctionWithContext.h b/Engine/Source/Runtime/Core/Public/Templates/FunctionWithContext.h
--- a/Engine/Source/Runtime/Core/Public/Templates/FunctionWithContext.h
+++ b/Engine/Source/Runtime/Core/Public/Templates/FunctionWithContext.h
@@ -99,7 +99,7 @@ public:
 
 private:
 	template <typename CallableType>
-	AUTORTFM_INFER static ReturnType Call(void* Callable, ArgTypes... Args)
+	static ReturnType Call(void* Callable, ArgTypes... Args)
 	{
 		return Invoke(*(CallableType*)Callable, (ArgTypes&&)Args...);
 	}
diff --git a/Engine/Source/Runtime/Core/Public/Templates/Invoke.h b/Engine/Source/Runtime/Core/Public/Templates/Invoke.h
--- a/Engine/Source/Runtime/Core/Public/Templates/Invoke.h
+++ b/Engine/Source/Runtime/Core/Public/Templates/Invoke.h
@@ -3,7 +3,6 @@
 #pragma once
 
 #include "CoreTypes.h"
-#include "AutoRTFM.h"
 #include "Traits/MemberPtrOuter.h"
 #include <type_traits>
 
@@ -41,14 +40,14 @@ namespace UE::Core::Private
  * See: http://en.cppreference.com/w/cpp/utility/functional/invoke
  */
 template <typename FuncType, typename... ArgTypes>
-AUTORTFM_INFER UE_FORCEINLINE_HINT constexpr auto Invoke(FuncType&& Func, ArgTypes&&... Args)
+UE_FORCEINLINE_HINT constexpr auto Invoke(FuncType&& Func, ArgTypes&&... Args)
 	-> decltype(((FuncType&&)Func)((ArgTypes&&)Args...))
 {
 	return ((FuncType&&)Func)((ArgTypes&&)Args...);
 }
 
 template <typename ReturnType, typename ObjType, typename TargetType>
-AUTORTFM_INFER constexpr auto Invoke(ReturnType ObjType::*pdm, TargetType&& Target)
+constexpr auto Invoke(ReturnType ObjType::*pdm, TargetType&& Target)
 	-> decltype(UE::Core::Private::DereferenceIfNecessary<ObjType>((TargetType&&)Target, &Target).*pdm)
 {
 	return UE::Core::Private::DereferenceIfNecessary<ObjType>((TargetType&&)Target, &Target).*pdm;
@@ -60,7 +59,7 @@ template <
 	typename... ArgTypes,
 	typename    ObjType = TMemberPtrOuter_T<PtrMemFunType>
 >
-AUTORTFM_INFER constexpr auto Invoke(PtrMemFunType PtrMemFun, TargetType&& Target, ArgTypes&&... Args)
+constexpr auto Invoke(PtrMemFunType PtrMemFun, TargetType&& Target, ArgTypes&&... Args)
 	-> decltype((UE::Core::Private::DereferenceIfNecessary<ObjType>((TargetType&&)Target, &Target).*PtrMemFun)((ArgTypes&&)Args...))
 {
 	return (UE::Core::Private::DereferenceIfNecessary<ObjType>((TargetType&&)Target, &Target).*PtrMemFun)((ArgTypes&&)Args...);
diff --git a/Engine/Source/Runtime/Core/Public/Templates/MemoryOps.h b/Engine/Source/Runtime/Core/Public/Templates/MemoryOps.h
--- a/Engine/Source/Runtime/Core/Public/Templates/MemoryOps.h
+++ b/Engine/Source/Runtime/Core/Public/Templates/MemoryOps.h
@@ -2,7 +2,6 @@
 
 #pragma once
 
-#include "AutoRTFM.h"
 #include "CoreTypes.h"
 #include "HAL/UnrealMemory.h"
 #include "Templates/UnrealTypeTraits.h"
@@ -54,7 +53,7 @@ FORCENOINLINE void DefaultConstructItems(void* Address, SizeType Count)
  * @note: This function is optimized for values of T, and so will not dynamically dispatch destructor calls if T's destructor is virtual.
  */
 template <typename ElementType>
-AUTORTFM_INFER FORCEINLINE constexpr void DestructItem(ElementType* Element)
+FORCEINLINE constexpr void DestructItem(ElementType* Element)
 {
 	if constexpr (sizeof(ElementType) == 0)
 	{
@@ -79,12 +78,12 @@ AUTORTFM_INFER FORCEINLINE constexpr void DestructItem(ElementType* Element)
  */
 template <typename ElementType, typename SizeType>
 	requires (sizeof(ElementType) > 0 && std::is_trivially_destructible_v<ElementType>) // the sizeof here should improve the error messages we get when we try to call this function with incomplete types
-AUTORTFM_INFER FORCEINLINE constexpr void DestructItems(ElementType* Element, SizeType Count)
+FORCEINLINE constexpr void DestructItems(ElementType* Element, SizeType Count)
 {
 }
 template <typename ElementType, typename SizeType>
 	requires (sizeof(ElementType) > 0 && !std::is_trivially_destructible_v<ElementType>) // the sizeof here should improve the error messages we get when we try to call this function with incomplete types
-AUTORTFM_INFER FORCENOINLINE constexpr void DestructItems(ElementType* Element, SizeType Count)
+FORCENOINLINE constexpr void DestructItems(ElementType* Element, SizeType Count)
 {
 	while (Count)
 	{
@@ -115,7 +114,7 @@ FORCEINLINE void ConstructItems(void* Dest, const SourceElementType* Source, Siz
 }
 template <typename DestinationElementType, typename SourceElementType, typename SizeType>
 	requires (sizeof(DestinationElementType) > 0 && sizeof(SourceElementType) > 0 && !::TIsBitwiseConstructible<DestinationElementType, SourceElementType>::Value) // the sizeof here should improve the error messages we get when we try to call this function with incomplete types
-AUTORTFM_INFER FORCENOINLINE void ConstructItems(void* Dest, const SourceElementType* Source, SizeType Count)
+FORCENOINLINE void ConstructItems(void* Dest, const SourceElementType* Source, SizeType Count)
 {
 	while (Count)
 	{
@@ -141,7 +140,7 @@ FORCEINLINE void CopyAssignItems(ElementType* Dest, const ElementType* Source, S
 }
 template <typename ElementType, typename SizeType>
 	requires (sizeof(ElementType) > 0 && !std::is_trivially_copy_assignable_v<ElementType>) // the sizeof here should improve the error messages we get when we try to call this function with incomplete types
-AUTORTFM_INFER FORCENOINLINE void CopyAssignItems(ElementType* Dest, const ElementType* Source, SizeType Count)
+FORCENOINLINE void CopyAssignItems(ElementType* Dest, const ElementType* Source, SizeType Count)
 {
 	while (Count)
 	{
@@ -160,7 +159,7 @@ AUTORTFM_INFER FORCENOINLINE void CopyAssignItems(ElementType* Dest, const Eleme
  * @param	Source		A pointer to the first item to relocate.
  */
 template <typename DestinationElementType, typename SourceElementType, typename SizeType>
-AUTORTFM_INFER FORCEINLINE void RelocateConstructItem(void* Dest, const SourceElementType* Source)
+FORCEINLINE void RelocateConstructItem(void* Dest, const SourceElementType* Source)
 {
 	if constexpr (sizeof(DestinationElementType) == 0 || sizeof(SourceElementType) == 0)
 	{
diff --git a/Engine/Source/Runtime/Core/Public/Templates/Projection.h b/Engine/Source/Runtime/Core/Public/Templates/Projection.h
--- a/Engine/Source/Runtime/Core/Public/Templates/Projection.h
+++ b/Engine/Source/Runtime/Core/Public/Templates/Projection.h
@@ -2,16 +2,15 @@
 
 #pragma once
 
-#include "AutoRTFM.h"
 #include <type_traits>
 
 namespace UE::Core::Private
 {
 	template <typename InvocableType>
-	struct AUTORTFM_INFER TProjectionMemberFunction;
+	struct TProjectionMemberFunction;
 
 	template <typename ClassType, typename FunctionType>
-	struct AUTORTFM_INFER TProjectionMemberFunction<FunctionType ClassType::*>
+	struct TProjectionMemberFunction<FunctionType ClassType::*>
 	{
 		FunctionType ClassType::* MemberFunctionPtr;
 
@@ -105,7 +104,7 @@ namespace UE::Core::Private
  * Algo::SortBy(ArrayOfOuters, Projection(&FOuter::Inner, &FInner::Name));
  */
 template <typename Invocable0Type, typename... InvocableTypes>
-AUTORTFM_INFER [[nodiscard]] constexpr auto Projection(Invocable0Type&& Invocable0, InvocableTypes&&... Invocables)
+[[nodiscard]] constexpr auto Projection(Invocable0Type&& Invocable0, InvocableTypes&&... Invocables)
 {
 	if constexpr (sizeof...(InvocableTypes) == 0)
 	{
diff --git a/Engine/Source/Runtime/Core/Public/Templates/SharedPointer.h b/Engine/Source/Runtime/Core/Public/Templates/SharedPointer.h
--- a/Engine/Source/Runtime/Core/Public/Templates/SharedPointer.h
+++ b/Engine/Source/Runtime/Core/Public/Templates/SharedPointer.h
@@ -2070,7 +2070,7 @@ template< class ObjectType >
  */
 // NOTE: The following is an Unreal extension to standard shared_ptr behavior
 template< class ObjectType, class DeleterType >
-AUTORTFM_INFER [[nodiscard]] inline SharedPointerInternals::TRawPtrProxyWithDeleter< ObjectType, DeleterType > MakeShareable( ObjectType* InObject, DeleterType&& InDeleter )
+[[nodiscard]] inline SharedPointerInternals::TRawPtrProxyWithDeleter< ObjectType, DeleterType > MakeShareable( ObjectType* InObject, DeleterType&& InDeleter )
 {
 	///////////////////////////////////////////////////////////////////////////////////////////////
 	// If you get an 'ambiguous call' compile error in this function, it means you have multiple //
@@ -2108,7 +2108,7 @@ AUTORTFM_INFER [[nodiscard]] inline SharedPointerInternals::TRawPtrProxyWithDele
  *     TSharedPtr<FMyType> Val = MakeShared<FMyType>(FMyType::FPrivateToken{}, 5, 3.14f, TEXT("Banana"));
  */
 template <typename InObjectType, ESPMode InMode = ESPMode::ThreadSafe, typename... InArgTypes>
-AUTORTFM_INFER [[nodiscard]] inline TSharedRef<InObjectType, InMode> MakeShared(InArgTypes&&... Args)
+[[nodiscard]] inline TSharedRef<InObjectType, InMode> MakeShared(InArgTypes&&... Args)
 {
 	SharedPointerInternals::TIntrusiveReferenceController<InObjectType, InMode>* Controller = SharedPointerInternals::NewIntrusiveReferenceController<InMode, InObjectType>(Forward<InArgTypes>(Args)...);
 	return UE::Core::Private::MakeSharedRef<InObjectType, InMode>(Controller->GetObjectPtr(), (SharedPointerInternals::TReferenceControllerBase<InMode>*)Controller);
diff --git a/Engine/Source/Runtime/Core/Public/Templates/SharedPointerInternals.h b/Engine/Source/Runtime/Core/Public/Templates/SharedPointerInternals.h
--- a/Engine/Source/Runtime/Core/Public/Templates/SharedPointerInternals.h
+++ b/Engine/Source/Runtime/Core/Public/Templates/SharedPointerInternals.h
@@ -340,7 +340,7 @@ namespace SharedPointerInternals
 	// it stores the custom deleter as a member to allow a function pointer to be used as a custom deleter (a function pointer can't 
 	// be a base class)
 	template <typename DeleterType, bool bIsZeroSize = std::is_empty_v<DeleterType>>
-	struct AUTORTFM_INFER TDeleterHolder : private DeleterType
+	struct TDeleterHolder : private DeleterType
 	{
 		explicit TDeleterHolder(DeleterType&& Arg)
 			: DeleterType(MoveTemp(Arg))
@@ -355,7 +355,7 @@ namespace SharedPointerInternals
 	};
 
 	template <typename DeleterType>
-	struct AUTORTFM_INFER TDeleterHolder<DeleterType, false>
+	struct TDeleterHolder<DeleterType, false>
 	{
 		explicit TDeleterHolder(DeleterType&& Arg)
 			: Deleter(MoveTemp(Arg))
@@ -373,7 +373,7 @@ namespace SharedPointerInternals
 	};
 
 	template <typename ObjectType, typename DeleterType, ESPMode Mode>
-	class AUTORTFM_INFER TReferenceControllerWithDeleter : private TDeleterHolder<DeleterType>, public TReferenceControllerBase<Mode>
+	class TReferenceControllerWithDeleter : private TDeleterHolder<DeleterType>, public TReferenceControllerBase<Mode>
 	{
 	public:
 		explicit TReferenceControllerWithDeleter(ObjectType* InObject, DeleterType&& Deleter)
@@ -397,7 +397,7 @@ namespace SharedPointerInternals
 	};
 
 	template <typename ObjectType, ESPMode Mode>
-	class AUTORTFM_INFER TIntrusiveReferenceController : public TReferenceControllerBase<Mode>
+	class TIntrusiveReferenceController : public TReferenceControllerBase<Mode>
 	{
 	public:
 		template <typename... ArgTypes>
@@ -448,7 +448,7 @@ namespace SharedPointerInternals
 
 	/** Deletes an object via the standard delete operator */
 	template <typename Type>
-	struct AUTORTFM_INFER DefaultDeleter
+	struct DefaultDeleter
 	{
 		UE_FORCEINLINE_HINT void operator()(Type* Object) const
 		{
@@ -458,21 +458,21 @@ namespace SharedPointerInternals
 
 	/** Creates a reference controller which just calls delete */
 	template <ESPMode Mode, typename ObjectType>
-	AUTORTFM_INFER inline TReferenceControllerBase<Mode>* NewDefaultReferenceController(ObjectType* Object)
+	inline TReferenceControllerBase<Mode>* NewDefaultReferenceController(ObjectType* Object)
 	{
 		return new TReferenceControllerWithDeleter<ObjectType, DefaultDeleter<ObjectType>, Mode>(Object, DefaultDeleter<ObjectType>());
 	}
 
 	/** Creates a custom reference controller with a specified deleter */
 	template <ESPMode Mode, typename ObjectType, typename DeleterType>
-	AUTORTFM_INFER inline TReferenceControllerBase<Mode>* NewCustomReferenceController(ObjectType* Object, DeleterType&& Deleter)
+	inline TReferenceControllerBase<Mode>* NewCustomReferenceController(ObjectType* Object, DeleterType&& Deleter)
 	{
 		return new TReferenceControllerWithDeleter<ObjectType, typename TRemoveReference<DeleterType>::Type, Mode>(Object, Forward<DeleterType>(Deleter));
 	}
 
 	/** Creates an intrusive reference controller */
 	template <ESPMode Mode, typename ObjectType, typename... ArgTypes>
-	AUTORTFM_INFER inline TIntrusiveReferenceController<ObjectType, Mode>* NewIntrusiveReferenceController(ArgTypes&&... Args)
+	inline TIntrusiveReferenceController<ObjectType, Mode>* NewIntrusiveReferenceController(ArgTypes&&... Args)
 	{
 		return new TIntrusiveReferenceController<ObjectType, Mode>(Forward<ArgTypes>(Args)...);
 	}
diff --git a/Engine/Source/Runtime/Core/Public/Templates/Tuple.h b/Engine/Source/Runtime/Core/Public/Templates/Tuple.h
--- a/Engine/Source/Runtime/Core/Public/Templates/Tuple.h
+++ b/Engine/Source/Runtime/Core/Public/Templates/Tuple.h
@@ -2,7 +2,6 @@
 
 #pragma once
 
-#include "AutoRTFM/Defines.h"
 #include "CoreTypes.h"
 #include "Concepts/NotCVRefTo.h"
 #include "Templates/UnrealTemplate.h"
@@ -57,7 +56,7 @@ namespace UE::Core::Private::Tuple
 	constexpr uint32 TTypeCountInParameterPack_V<T, U, Types...> = TTypeCountInParameterPack_V<T, Types...> + (std::is_same_v<T, U> ? 1 : 0);
 
 	template <typename T, uint32 Index, uint32 TupleSize>
-	struct AUTORTFM_INFER TTupleBaseElement
+	struct TTupleBaseElement
 	{
 		template <typename ArgType>
 		constexpr explicit TTupleBaseElement(EForwardingConstructor, ArgType&& Arg)
@@ -223,7 +222,7 @@ namespace UE::Core::Private::Tuple
 	struct TTupleBase;
 
 	template <uint32... Indices, typename... Types>
-	struct AUTORTFM_INFER TTupleBase<TIntegerSequence<uint32, Indices...>, Types...> : TTupleBaseElement<Types, Indices, sizeof...(Types)>...
+	struct TTupleBase<TIntegerSequence<uint32, Indices...>, Types...> : TTupleBaseElement<Types, Indices, sizeof...(Types)>...
 	{
 		template <typename... ArgTypes>
 		constexpr explicit TTupleBase(EForwardingConstructor, ArgTypes&&... Args)
@@ -529,7 +528,7 @@ namespace UE::Core::Private::Tuple
 }
 
 template <typename... Types>
-struct AUTORTFM_INFER TTuple : UE::Core::Private::Tuple::TTupleBase<TMakeIntegerSequence<uint32, sizeof...(Types)>, Types...>
+struct TTuple : UE::Core::Private::Tuple::TTupleBase<TMakeIntegerSequence<uint32, sizeof...(Types)>, Types...>
 {
 private:
 	typedef UE::Core::Private::Tuple::TTupleBase<TMakeIntegerSequence<uint32, sizeof...(Types)>, Types...> Super;
diff --git a/Engine/Source/Runtime/Core/Public/Templates/UniquePtr.h b/Engine/Source/Runtime/Core/Public/Templates/UniquePtr.h
--- a/Engine/Source/Runtime/Core/Public/Templates/UniquePtr.h
+++ b/Engine/Source/Runtime/Core/Public/Templates/UniquePtr.h
@@ -2,7 +2,6 @@
 
 #pragma once
 
-#include "AutoRTFM/Defines.h"
 #include "CoreTypes.h"
 #include "Templates/UnrealTemplate.h"
 #include "Templates/IsArray.h"
@@ -25,7 +24,7 @@
 // https://en.cppreference.com/w/cpp/memory/unique_ptr/~unique_ptr
 
 template <typename T>
-struct AUTORTFM_INFER TDefaultDelete
+struct TDefaultDelete
 {
 	DECLARE_INLINE_TYPE_LAYOUT(TDefaultDelete, NonVirtual);
 
@@ -104,7 +103,7 @@ struct TDefaultDelete<T[]>
 };
 
 template <typename T, typename Deleter = TDefaultDelete<T>>
-class AUTORTFM_INFER TUniquePtr : private Deleter
+class TUniquePtr : private Deleter
 {
 	DECLARE_INLINE_TYPE_LAYOUT_EXPLICIT_BASES(TUniquePtr, NonVirtual, Deleter);
 
@@ -497,7 +496,7 @@ public:
 };
 
 template <typename T, typename Deleter>
-class AUTORTFM_INFER TUniquePtr<T[], Deleter> : private Deleter
+class TUniquePtr<T[], Deleter> : private Deleter
 {
 	template <typename OtherT, typename OtherDeleter>
 	friend class TUniquePtr;
