Fix FAutoConsoleObject's destructor keeping the console objects live in the console manager after they're gone

- This causes crashes with merged modules because after a library unmounts and remounts, pointers to static data are invalidated
- CVar objects keep a reference to the default value which ends up compiled into the library
- Thus we need to enforce full destruction of those objects

#Synced-CL 52374856

diff --git a/Engine/Source/Runtime/Core/Public/HAL/IConsoleManager.h b/Engine/Source/Runtime/Core/Public/HAL/IConsoleManager.h
--- a/Engine/Source/Runtime/Core/Public/HAL/IConsoleManager.h
+++ b/Engine/Source/Runtime/Core/Public/HAL/IConsoleManager.h
@@ -1502,7 +1502,8 @@ public:
 	/** Destructor, removes the console object **/
 	virtual ~FAutoConsoleObject()
 	{
-		IConsoleManager::Get().UnregisterConsoleObject(Target);
+		const bool bKeepCVarObject = false; // do not keep the console object in memory, as it may cause crashes if we reload the object later
+		IConsoleManager::Get().UnregisterConsoleObject(Target, bKeepCVarObject);
 	}
 
 	static CORE_API TArray<const FAutoConsoleObject*>& AccessGeneralShaderChangeCvars();
