Added FText::AsMemory overload accepting a specific EMemoryUnit.

The existing AsMemory always auto-selects the display unit based on magnitude.
This adds a new EMemoryUnit enum (KiB, MiB, GiB, etc. and KB, MB, GB, etc.) and an overload that formats to a caller-specified unit instead.
This allows always displaying things like "X MiB" regardless of the byte count. EMemoryUnit::Auto preserves the existing behavior.

[RN] Added FText::AsMemory overload accepting a specific EMemoryUnit
#Synced-CL 52109349

diff --git a/Engine/Source/Runtime/Core/Private/Internationalization/Text.cpp b/Engine/Source/Runtime/Core/Private/Internationalization/Text.cpp
--- a/Engine/Source/Runtime/Core/Private/Internationalization/Text.cpp
+++ b/Engine/Source/Runtime/Core/Private/Internationalization/Text.cpp
@@ -834,6 +834,47 @@ FText FText::AsMemory(uint64 NumBytes, EMemoryUnitStandard UnitStandard)
 	return FText::AsMemory(NumBytes, nullptr, nullptr, UnitStandard);
 }
 
+FText FText::AsMemory(uint64 NumBytes, EMemoryUnit Unit, const FNumberFormattingOptions* const Options, const FCulturePtr& TargetCulture)
+{
+	if (Unit == EMemoryUnit::Auto)
+	{
+		return FText::AsMemory(NumBytes, Options, TargetCulture);
+	}
+
+	checkf(FInternationalization::Get().IsInitialized() == true, TEXT("FInternationalization is not initialized. An FText formatting method was likely used in static object initialization - this is not supported."));
+
+	if (Unit == EMemoryUnit::Unit_Bytes)
+	{
+		FFormatNamedArguments Args;
+		Args.Add(TEXT("Number"), FText::AsNumber(NumBytes, Options, TargetCulture));
+		Args.Add(TEXT("Unit"), FText::FromString(TEXT("B")));
+		return FText::Format(NSLOCTEXT("Internationalization", "ComputerMemoryFormatting", "{Number} {Unit}"), Args);
+	}
+
+	static const TCHAR* IECPrefixes = TEXT("KMGTPE");
+	static const TCHAR* SIPrefixes  = TEXT("kMGTPE");
+
+	const bool bIEC = (Unit <= EMemoryUnit::Unit_EiB);
+	const int32 PrefixIndex = bIEC
+		? (int32)Unit - (int32)EMemoryUnit::Unit_KiB
+		: (int32)Unit - (int32)EMemoryUnit::Unit_KB;
+	const uint64 UnitDivisor = bIEC ? 1024ull : 1000ull;
+	const TCHAR* Prefixes = bIEC ? IECPrefixes : SIPrefixes;
+	const TCHAR* Suffix = bIEC ? TEXT("iB") : TEXT("B");
+
+	uint64 Divisor = 1;
+	for (int32 i = 0; i <= PrefixIndex; ++i)
+	{
+		Divisor *= UnitDivisor;
+	}
+
+	const double Value = (double)NumBytes / (double)Divisor;
+	FFormatNamedArguments Args;
+	Args.Add(TEXT("Number"), FText::AsNumber(Value, Options, TargetCulture));
+	Args.Add(TEXT("Unit"), FText::FromString(FString::ConstructFromPtrSize(&Prefixes[PrefixIndex], 1) + Suffix));
+	return FText::Format(NSLOCTEXT("Internationalization", "ComputerMemoryFormatting", "{Number} {Unit}"), Args);
+}
+
 FString FText::GetInvariantTimeZone()
 {
 	return TEXT("Etc/Unknown");
diff --git a/Engine/Source/Runtime/Core/Public/Internationalization/Text.h b/Engine/Source/Runtime/Core/Public/Internationalization/Text.h
--- a/Engine/Source/Runtime/Core/Public/Internationalization/Text.h
+++ b/Engine/Source/Runtime/Core/Public/Internationalization/Text.h
@@ -196,6 +196,18 @@ enum EMemoryUnitStandard : int
 	SI
 };
 
+/** Specific memory unit to use when formatting memory sizes. IEC units are 1024-based (KiB, MiB, ...), SI units are 1000-based (KB, MB, ...). */
+enum class EMemoryUnit : int
+{
+	// Auto-select based on magnitude (default behavior)
+	Auto,
+	Unit_Bytes,
+	// IEC (1024-based)
+	Unit_KiB, Unit_MiB, Unit_GiB, Unit_TiB, Unit_PiB, Unit_EiB,
+	// SI (1000-based)
+	Unit_KB, Unit_MB, Unit_GB, Unit_TB, Unit_PB, Unit_EB
+};
+
 struct FNumberFormattingOptions
 {
 	CORE_API FNumberFormattingOptions();
@@ -476,6 +488,12 @@ public:
 	 */
 	[[nodiscard]] static CORE_API FText AsMemory(uint64 NumBytes, EMemoryUnitStandard UnitStandard);
 
+	/**
+	 * Generate an FText that represents the passed number as a memory size in the current culture, expressed in a specific unit.
+	 * Use EMemoryUnit::Auto to preserve the default magnitude-based selection behavior.
+	 */
+	[[nodiscard]] static CORE_API FText AsMemory(uint64 NumBytes, EMemoryUnit Unit, const FNumberFormattingOptions* const Options = nullptr, const FCulturePtr& TargetCulture = nullptr);
+
 	/**
 	 * Attempts to create an FText instance from a string table ID and key (this is the same as the LOCTABLE macro, except this can also work with non-literal string values).
 	 * @return The found text, or a dummy FText if not found.
diff --git a/Engine/Source/Runtime/Core/Tests/Internationalization/TextTest.cpp b/Engine/Source/Runtime/Core/Tests/Internationalization/TextTest.cpp
--- a/Engine/Source/Runtime/Core/Tests/Internationalization/TextTest.cpp
+++ b/Engine/Source/Runtime/Core/Tests/Internationalization/TextTest.cpp
@@ -377,6 +377,49 @@ public:
 			}
 #undef TEST
 
+#define TEST_UNIT(NumBytes, Unit, ExpectedString) \
+		if (!FText::FromString(TEXT(ExpectedString)).EqualTo(FText::AsMemory(NumBytes, Unit, &NumberFormattingOptions, nullptr))) \
+		{ \
+			FAIL_CHECK(FString::Printf(TEXT("FText::AsMemory(EMemoryUnit) expected %s bytes in %s to be %s - got %s"), TEXT(#NumBytes), TEXT(#Unit), TEXT(ExpectedString), *FText::AsMemory(NumBytes, Unit, &NumberFormattingOptions, nullptr).ToString())); \
+		} \
+
+			{
+				FNumberFormattingOptions NumberFormattingOptions = FNumberFormattingOptions()
+					.SetRoundingMode(ERoundingMode::HalfFromZero)
+					.SetMinimumFractionalDigits(0)
+					.SetMaximumFractionalDigits(3);
+
+				// Auto delegates to existing IEC behavior
+				TEST_UNIT(0, EMemoryUnit::Auto, "0 B");
+				TEST_UNIT(500, EMemoryUnit::Auto, "500 B");
+				TEST_UNIT(1048576, EMemoryUnit::Auto, "1 MiB");
+
+				// Explicit Bytes
+				TEST_UNIT(0, EMemoryUnit::Unit_Bytes, "0 B");
+				TEST_UNIT(1, EMemoryUnit::Unit_Bytes, "1 B");
+				TEST_UNIT(999, EMemoryUnit::Unit_Bytes, "999 B");
+
+				// IEC units - exact powers of 1024
+				TEST_UNIT(1024, EMemoryUnit::Unit_KiB, "1 KiB");
+				TEST_UNIT(1048576, EMemoryUnit::Unit_MiB, "1 MiB");
+				TEST_UNIT(1073741824, EMemoryUnit::Unit_GiB, "1 GiB");
+				TEST_UNIT(1099511627776, EMemoryUnit::Unit_TiB, "1 TiB");
+				TEST_UNIT(1125899906842624, EMemoryUnit::Unit_PiB, "1 PiB");
+				TEST_UNIT(1152921504606846976, EMemoryUnit::Unit_EiB, "1 EiB");
+
+				// SI units - exact powers of 1000
+				TEST_UNIT(1000, EMemoryUnit::Unit_KB, "1 kB");
+				TEST_UNIT(1000000, EMemoryUnit::Unit_MB, "1 MB");
+				TEST_UNIT(1000000000, EMemoryUnit::Unit_GB, "1 GB");
+				TEST_UNIT(1000000000000, EMemoryUnit::Unit_TB, "1 TB");
+				TEST_UNIT(1000000000000000, EMemoryUnit::Unit_PB, "1 PB");
+				TEST_UNIT(1000000000000000000, EMemoryUnit::Unit_EB,  "1 EB");
+
+				// max uint64: float-based division yields 16 EiB (vs 15.999 from the integer-division path in the old overload)
+				TEST_UNIT(18446744073709551615ULL, EMemoryUnit::Unit_EiB, "16 EiB");
+			}
+#undef TEST_UNIT
+
 #define TEST( A, B, ComparisonLevel ) if( !(FText::FromString(A)).EqualTo(FText::FromString(B), (ComparisonLevel)) ) FAIL_CHECK(FString::Printf(TEXT("Testing comparison of equivalent characters with comparison level (%s). - A=%s B=%s"),TEXT(#ComparisonLevel),(A),(B)))
 
 			// Basic sanity checks
