Core: Added UE::CInsertable as a proper concept and UE::CInsertableInto that reverses the parameters

#Synced-CL 50967057

diff --git a/Engine/Source/Runtime/Core/Public/Concepts/Insertable.h b/Engine/Source/Runtime/Core/Public/Concepts/Insertable.h
--- a/Engine/Source/Runtime/Core/Public/Concepts/Insertable.h
+++ b/Engine/Source/Runtime/Core/Public/Concepts/Insertable.h
@@ -12,3 +12,29 @@ struct CInsertable {
 		Dest << Val
 	);
 };
+
+namespace UE
+{
+	/**
+	 * Describes inserting an instance of a value type into a destination type using operator<<.
+	 *
+	 * Equivalent to CInsertableInto<ValueType, DestType>.
+	 * Read as "DestType is insertable from ValueType".
+	 * Constrain the value type with <CInsertable<ValueType> DestType>.
+	 */
+	template <typename DestType, typename ValueType>
+	concept CInsertable = requires(DestType& Dest, ValueType& Val)
+	{
+		Dest << Val;
+	};
+
+	/**
+	 * Describes inserting an instance of a value type into a destination type using operator<<.
+	 *
+	 * Equivalent to CInsertable<DestType, ValueType>.
+	 * Read as "ValueType is insertable into DestType".
+	 * Constrain the destination type with <CInsertableInto<DestType> ValueType>.
+	 */
+	template <typename ValueType, typename DestType>
+	concept CInsertableInto = CInsertable<DestType, ValueType>;
+}
diff --git a/Engine/Source/Runtime/Core/Public/Logging/StructuredLog.h b/Engine/Source/Runtime/Core/Public/Logging/StructuredLog.h
--- a/Engine/Source/Runtime/Core/Public/Logging/StructuredLog.h
+++ b/Engine/Source/Runtime/Core/Public/Logging/StructuredLog.h
@@ -293,7 +293,7 @@ UE_API void VisitLogContext(TFunctionRef<void (const FCbField&)> Visitor);
  *
  * Arrays and objects without custom formatting are converted to JSON.
  */
-template <typename ValueType UE_REQUIRES(TModels_V<CInsertable<FCbWriter&>, ValueType>)>
+template <CInsertableInto<FCbWriter&> ValueType>
 inline void SerializeForLog(FCbWriter& Writer, ValueType&& Value)
 {
 	Writer << (ValueType&&)Value;
diff --git a/Engine/Source/Runtime/Core/Public/Misc/TVariantMeta.h b/Engine/Source/Runtime/Core/Public/Misc/TVariantMeta.h
--- a/Engine/Source/Runtime/Core/Public/Misc/TVariantMeta.h
+++ b/Engine/Source/Runtime/Core/Public/Misc/TVariantMeta.h
@@ -262,7 +262,7 @@ namespace Private
 	{
 		using VariantType = TVariant<Ts...>;
 		static_assert((std::is_default_constructible_v<Ts> && ...), "Each type in TVariant template parameter pack must be default constructible in order to use FArchive serialization");
-		static_assert((TModels_V<CInsertable<FArchive&>, Ts> && ...), "Each type in TVariant template parameter pack must be able to use operator<< with an FArchive");
+		static_assert((UE::CInsertable<FArchive&, Ts> && ...), "Each type in TVariant template parameter pack must be able to use operator<< with an FArchive");
 
 		/** Load the type at the specified index from the FArchive and emplace it into the TVariant */
 		static void Load(SIZE_T TypeIndex, FArchive& Ar, VariantType& OutVariant)
