Core: Fixed comments on UE::CInsertable[Into]

#rnx
#Synced-CL 50967057

diff --git a/Engine/Source/Runtime/Core/Public/Concepts/Insertable.h b/Engine/Source/Runtime/Core/Public/Concepts/Insertable.h
--- a/Engine/Source/Runtime/Core/Public/Concepts/Insertable.h
+++ b/Engine/Source/Runtime/Core/Public/Concepts/Insertable.h
@@ -20,7 +20,7 @@ namespace UE
 	 *
 	 * Equivalent to CInsertableInto<ValueType, DestType>.
 	 * Read as "DestType is insertable from ValueType".
-	 * Constrain the value type with <CInsertable<ValueType> DestType>.
+	 * Constrain the destination type with <CInsertable<ValueType> DestType>.
 	 */
 	template <typename DestType, typename ValueType>
 	concept CInsertable = requires(DestType& Dest, ValueType& Val)
@@ -33,7 +33,7 @@ namespace UE
 	 *
 	 * Equivalent to CInsertable<DestType, ValueType>.
 	 * Read as "ValueType is insertable into DestType".
-	 * Constrain the destination type with <CInsertableInto<DestType> ValueType>.
+	 * Constrain the value type with <CInsertableInto<DestType> ValueType>.
 	 */
 	template <typename ValueType, typename DestType>
 	concept CInsertableInto = CInsertable<DestType, ValueType>;
