Add descriptions to trace channels and a few other places.
#Synced-CL 51749465

diff --git a/Engine/Plugins/Animation/PoseSearch/Source/Runtime/Private/Trace/PoseSearchTraceLogger.cpp b/Engine/Plugins/Animation/PoseSearch/Source/Runtime/Private/Trace/PoseSearchTraceLogger.cpp
--- a/Engine/Plugins/Animation/PoseSearch/Source/Runtime/Private/Trace/PoseSearchTraceLogger.cpp
+++ b/Engine/Plugins/Animation/PoseSearch/Source/Runtime/Private/Trace/PoseSearchTraceLogger.cpp
@@ -10,7 +10,9 @@
 #include "Trace/Trace.inl"
 #include "TraceFilter.h"
 
-UE_TRACE_CHANNEL_DEFINE(PoseSearchChannel);
+UE_TRACE_CHANNEL_DEFINE(PoseSearchChannel, "Traces motion matching animation state data for the PoseSearch system. Captures serialized pose data, \
+velocities, search costs, database information, and real-time animation search operations for debugging animation \
+matching decisions.");
 
 // No custom versions applied when loading.
 UE_TRACE_EVENT_BEGIN(PoseSearch, MotionMatchingState)
diff --git a/Engine/Plugins/Developer/Concert/ConcertMain/Source/ConcertTransport/Private/Trace/ConcertTrace.cpp b/Engine/Plugins/Developer/Concert/ConcertMain/Source/ConcertTransport/Private/Trace/ConcertTrace.cpp
--- a/Engine/Plugins/Developer/Concert/ConcertMain/Source/ConcertTransport/Private/Trace/ConcertTrace.cpp
+++ b/Engine/Plugins/Developer/Concert/ConcertMain/Source/ConcertTransport/Private/Trace/ConcertTrace.cpp
@@ -2,4 +2,5 @@
 
 #include "Trace/ConcertTrace.h"
 
-UE_TRACE_CHANNEL_DEFINE(ConcertChannel);
\ No newline at end of file
+UE_TRACE_CHANNEL_DEFINE(ConcertChannel, "Traces Concert multi-user collaboration events including protocol operations, network transport, \
+transaction synchronization, connection initialization, and collaborative editing operations across connected users.");
\ No newline at end of file
diff --git a/Engine/Plugins/FX/Niagara/Source/Niagara/Private/NiagaraTrace.cpp b/Engine/Plugins/FX/Niagara/Source/Niagara/Private/NiagaraTrace.cpp
--- a/Engine/Plugins/FX/Niagara/Source/Niagara/Private/NiagaraTrace.cpp
+++ b/Engine/Plugins/FX/Niagara/Source/Niagara/Private/NiagaraTrace.cpp
@@ -7,7 +7,8 @@
 
 #include "Particles/ParticlePerfStats.h"
 
-UE_TRACE_CHANNEL_DEFINE(NiagaraChannel)
+UE_TRACE_CHANNEL_DEFINE(NiagaraChannel, "Traces Niagara particle system operations including compilation, emission, system state, emitter state, \
+data channel operations, and particle data for debugging and performance analysis of VFX systems.");
 
 #if WITH_NIAGARA_INSIGHTS
 namespace NiagaraInsightsPrivate
diff --git a/Engine/Plugins/Media/PixelStreaming/Source/PixelStreaming/Private/PixelStreamingTrace.cpp b/Engine/Plugins/Media/PixelStreaming/Source/PixelStreaming/Private/PixelStreamingTrace.cpp
--- a/Engine/Plugins/Media/PixelStreaming/Source/PixelStreaming/Private/PixelStreamingTrace.cpp
+++ b/Engine/Plugins/Media/PixelStreaming/Source/PixelStreaming/Private/PixelStreamingTrace.cpp
@@ -4,4 +4,6 @@
 
 #include "Trace/Trace.inl"
 
-UE_TRACE_CHANNEL_DEFINE(PixelStreamingChannel);
+UE_TRACE_CHANNEL_DEFINE(PixelStreamingChannel, "Traces PixelStreaming (first generation) remote rendering and streaming operations including video \
+encoding / decoding (hardware / software / VPX), audio processing, streaming state, and performance metrics for debugging \
+streaming sessions.");
diff --git a/Engine/Plugins/Media/PixelStreaming2/Source/PixelStreaming2/Private/PixelStreaming2Trace.cpp b/Engine/Plugins/Media/PixelStreaming2/Source/PixelStreaming2/Private/PixelStreaming2Trace.cpp
--- a/Engine/Plugins/Media/PixelStreaming2/Source/PixelStreaming2/Private/PixelStreaming2Trace.cpp
+++ b/Engine/Plugins/Media/PixelStreaming2/Source/PixelStreaming2/Private/PixelStreaming2Trace.cpp
@@ -4,4 +4,5 @@
 
 #include "Trace/Trace.inl"
 
-UE_TRACE_CHANNEL_DEFINE(PixelStreaming2Channel);
+UE_TRACE_CHANNEL_DEFINE(PixelStreaming2Channel, "Traces PixelStreaming2 operations including video/audio streaming, video source/sink management, thread \
+operations, capture pipeline, and streaming performance metrics for debugging client streaming sessions.");
diff --git a/Engine/Plugins/Runtime/NetworkPrediction/Source/NetworkPrediction/Private/NetworkPredictionTrace.cpp b/Engine/Plugins/Runtime/NetworkPrediction/Source/NetworkPrediction/Private/NetworkPredictionTrace.cpp
--- a/Engine/Plugins/Runtime/NetworkPrediction/Source/NetworkPrediction/Private/NetworkPredictionTrace.cpp
+++ b/Engine/Plugins/Runtime/NetworkPrediction/Source/NetworkPrediction/Private/NetworkPredictionTrace.cpp
@@ -23,7 +23,9 @@ namespace NetworkPredictionTraceInternal
 	static constexpr ENetworkPredictionTraceVersion NetworkPredictionTraceVersion = ENetworkPredictionTraceVersion::Initial;
 };
 
-UE_TRACE_CHANNEL_DEFINE(NetworkPredictionChannel)
+UE_TRACE_CHANNEL_DEFINE(NetworkPredictionChannel, "Traces network prediction system operations including simulation state, client/server reconciliation, \
+rollback injection, input commands, sync state, auxiliary state, physics state, buffered input, fixed tick offsets, \
+out-of-band state modifications, and system faults for debugging networked gameplay prediction.")
 
 UE_TRACE_EVENT_BEGIN(NetworkPrediction, SimScope)
 	UE_TRACE_EVENT_FIELD(int32, TraceID)
diff --git a/Engine/Plugins/Runtime/RigVM/Source/RigVM/Private/RigVMTrace.cpp b/Engine/Plugins/Runtime/RigVM/Source/RigVM/Private/RigVMTrace.cpp
--- a/Engine/Plugins/Runtime/RigVM/Source/RigVM/Private/RigVMTrace.cpp
+++ b/Engine/Plugins/Runtime/RigVM/Source/RigVM/Private/RigVMTrace.cpp
@@ -8,7 +8,10 @@
 
 #if RIGVM_TRACE_ENABLED
 
-UE_TRACE_CHANNEL_DEFINE(RigVMChannel)
+UE_TRACE_CHANNEL_DEFINE(RigVMChannel, "Traces RigVM (Rig Virtual Machine) execution and variable state for rewind debugging including literal \
+memory, work memory, debug memory, external memory, draw interface, instruction visit info, and profiling data for \
+debugging rig evaluation and control rig behavior.")
+
 UE_TRACE_EVENT_DEFINE(RigVM, Literals);
 UE_TRACE_EVENT_DEFINE(RigVM, Execute);
 
diff --git a/Engine/Plugins/Runtime/StateTree/Source/StateTreeModule/Private/Debugger/StateTreeTrace.cpp b/Engine/Plugins/Runtime/StateTree/Source/StateTreeModule/Private/Debugger/StateTreeTrace.cpp
--- a/Engine/Plugins/Runtime/StateTree/Source/StateTreeModule/Private/Debugger/StateTreeTrace.cpp
+++ b/Engine/Plugins/Runtime/StateTree/Source/StateTreeModule/Private/Debugger/StateTreeTrace.cpp
@@ -21,7 +21,9 @@
 #include "Editor.h"
 #endif // WITH_EDITOR
 
-UE_TRACE_CHANNEL_DEFINE(StateTreeDebugChannel)
+UE_TRACE_CHANNEL_DEFINE(StateTreeDebugChannel, "Traces StateTree state machine execution for debugging including world timestamps, asset IDs, instance \
+events, frame events, phase transitions, log messages, state changes, task execution, evaluator execution, transition \
+conditions, condition evaluation, and active state tracking with deferred buffering support.")
 
 UE_TRACE_EVENT_BEGIN(StateTreeDebugger, WorldTimestampEvent)
 	UE_TRACE_EVENT_FIELD(double, WorldTime)
diff --git a/Engine/Plugins/TraceUtilities/Source/EditorTraceUtilities/Private/InsightsStatusBar/SInsightsStatusBar.cpp b/Engine/Plugins/TraceUtilities/Source/EditorTraceUtilities/Private/InsightsStatusBar/SInsightsStatusBar.cpp
--- a/Engine/Plugins/TraceUtilities/Source/EditorTraceUtilities/Private/InsightsStatusBar/SInsightsStatusBar.cpp
+++ b/Engine/Plugins/TraceUtilities/Source/EditorTraceUtilities/Private/InsightsStatusBar/SInsightsStatusBar.cpp
@@ -300,7 +300,7 @@ TSharedRef<SWidget> SInsightsStatusBarWidget::MakeTraceMenu()
 
 		MenuBuilder.AddMenuEntry(
 			LOCTEXT("StatNamedEventsLabel", "Stat Named Events"),
-			LOCTEXT("StatNamedEventsDesc", "Enable or disable named events in the stats system."),
+			LOCTEXT("StatNamedEventsDesc", "Significantly increases the number of CPU profiler events by activating tracing for the named events system."),
 			FSlateIcon(),
 			FUIAction(FExecuteAction::CreateLambda([]() { GCycleStatsShouldEmitNamedEvents = GCycleStatsShouldEmitNamedEvents == 0 ? 1 : 0; }),
 				      FCanExecuteAction(),
diff --git a/Engine/Plugins/VirtualProduction/LevelSnapshots/Source/LevelSnapshots/Private/LevelSnapshotsLog.cpp b/Engine/Plugins/VirtualProduction/LevelSnapshots/Source/LevelSnapshots/Private/LevelSnapshotsLog.cpp
--- a/Engine/Plugins/VirtualProduction/LevelSnapshots/Source/LevelSnapshots/Private/LevelSnapshotsLog.cpp
+++ b/Engine/Plugins/VirtualProduction/LevelSnapshots/Source/LevelSnapshots/Private/LevelSnapshotsLog.cpp
@@ -2,4 +2,4 @@
 
 #include "LevelSnapshotsLog.h"
 
-UE_TRACE_CHANNEL_DEFINE(LevelSnapshotsChannel);
\ No newline at end of file
+UE_TRACE_CHANNEL_DEFINE(LevelSnapshotsChannel, "Traces Level Snapshots operations for level state management including snapshot creation, comparison, restoration operations, and level state change tracking.");
\ No newline at end of file
diff --git a/Engine/Plugins/VirtualProduction/Takes/Source/TakeRecorder/Private/TakeRecorderTracing.cpp b/Engine/Plugins/VirtualProduction/Takes/Source/TakeRecorder/Private/TakeRecorderTracing.cpp
--- a/Engine/Plugins/VirtualProduction/Takes/Source/TakeRecorder/Private/TakeRecorderTracing.cpp
+++ b/Engine/Plugins/VirtualProduction/Takes/Source/TakeRecorder/Private/TakeRecorderTracing.cpp
@@ -2,4 +2,5 @@
 
 #include "TakeRecorderTracing.h"
 
-UE_TRACE_CHANNEL_DEFINE(TakeRecorderChannel);
+UE_TRACE_CHANNEL_DEFINE(TakeRecorderChannel, "Traces take recording operations for virtual production workflows. Used to track recording sessions, \
+playback operations, and take management activities.");
diff --git a/Engine/Source/Developer/TraceTools/Private/Models/TraceControlCommands.h b/Engine/Source/Developer/TraceTools/Private/Models/TraceControlCommands.h
--- a/Engine/Source/Developer/TraceTools/Private/Models/TraceControlCommands.h
+++ b/Engine/Source/Developer/TraceTools/Private/Models/TraceControlCommands.h
@@ -40,7 +40,7 @@ public:
 		UI_COMMAND(TraceBookmark, "Bookmark", "Trace a bookmark with a predefined name.", EUserInterfaceActionType::Button, FInputChord());
 		UI_COMMAND(TraceScreenshot, "Screenshot", "Trace a screenshot with a predefined name.", EUserInterfaceActionType::Button, FInputChord());
 
-		UI_COMMAND(ToggleStatNamedEvents, "Stat Named Events", "Enables/Disables the tracing of stat named events.", EUserInterfaceActionType::ToggleButton, FInputChord());
+		UI_COMMAND(ToggleStatNamedEvents, "Stat Named Events", "Significantly increases the number of CPU profiler events by activating tracing for the named events system.", EUserInterfaceActionType::ToggleButton, FInputChord());
 	}
 
 public:
diff --git a/Engine/Source/Editor/CurveEditor/Private/CurveEditorTrace.cpp b/Engine/Source/Editor/CurveEditor/Private/CurveEditorTrace.cpp
--- a/Engine/Source/Editor/CurveEditor/Private/CurveEditorTrace.cpp
+++ b/Engine/Source/Editor/CurveEditor/Private/CurveEditorTrace.cpp
@@ -1,7 +1,7 @@
-﻿// Copyright Epic Games, Inc. All Rights Reserved.
+// Copyright Epic Games, Inc. All Rights Reserved.
 
 #include "CurveEditorTrace.h"
 
 #include "Trace/Detail/Channel.h"
 
-UE_TRACE_CHANNEL_DEFINE(CurveEditorChannel);
+UE_TRACE_CHANNEL_DEFINE(CurveEditorChannel, "Curve Editor-specific CPU profiling events. Tracks performance of curve editing operations including key modifications, curve snapshots, diff calculations, and selection operations.");
diff --git a/Engine/Source/Programs/AutoRTFMTests/Private/InsightsTests.cpp b/Engine/Source/Programs/AutoRTFMTests/Private/InsightsTests.cpp
--- a/Engine/Source/Programs/AutoRTFMTests/Private/InsightsTests.cpp
+++ b/Engine/Source/Programs/AutoRTFMTests/Private/InsightsTests.cpp
@@ -14,7 +14,7 @@
 #include "Trace/Trace.h"
 #include "Trace/Trace.inl"
 
-UE_TRACE_CHANNEL_DEFINE(AutoRTFMInsightsChannel)
+UE_TRACE_CHANNEL_DEFINE(AutoRTFMInsightsChannel, "AutoRTFM transaction testing and profiling.")
 
 struct AutoRTFMInsightsDeferer final
 {
diff --git a/Engine/Source/Runtime/Core/Private/Async/TaskTrace.cpp b/Engine/Source/Runtime/Core/Private/Async/TaskTrace.cpp
--- a/Engine/Source/Runtime/Core/Private/Async/TaskTrace.cpp
+++ b/Engine/Source/Runtime/Core/Private/Async/TaskTrace.cpp
@@ -11,7 +11,7 @@
 
 namespace TaskTrace
 {
-	UE_TRACE_CHANNEL_DEFINE(TaskChannel);
+	UE_TRACE_CHANNEL_DEFINE(TaskChannel, "Task system lifecycle, scheduling, and dependencies. Data can be viewed in the Timing Profiler (Task Graph Insights).");
 
 	UE_TRACE_EVENT_BEGIN(TaskTrace, Init)
 		UE_TRACE_EVENT_FIELD(uint32, Version)
diff --git a/Engine/Source/Runtime/Core/Private/ProfilingDebugging/CallstackTrace.cpp b/Engine/Source/Runtime/Core/Private/ProfilingDebugging/CallstackTrace.cpp
--- a/Engine/Source/Runtime/Core/Private/ProfilingDebugging/CallstackTrace.cpp
+++ b/Engine/Source/Runtime/Core/Private/ProfilingDebugging/CallstackTrace.cpp
@@ -12,7 +12,7 @@ void	CallstackTrace_CreateInternal(FMalloc*);
 void	CallstackTrace_InitializeInternal();
 
 ////////////////////////////////////////////////////////////////////////////////
-UE_TRACE_CHANNEL_DEFINE(CallstackChannel)
+UE_TRACE_CHANNEL_DEFINE(CallstackChannel, "Allows allocations to be associated with callstacks. Requires Module channel to be enabled for symbol resolution to be possible.")
 UE_TRACE_EVENT_DEFINE(Memory, CallstackSpec)
 UE_TRACE_EVENT_DEFINE(Memory, CallstackSpecXORAndRLE)
 UE_TRACE_EVENT_DEFINE(Memory, CallstackSpecDelta7bit)
diff --git a/Engine/Source/Runtime/Core/Private/ProfilingDebugging/CountersTrace.cpp b/Engine/Source/Runtime/Core/Private/ProfilingDebugging/CountersTrace.cpp
--- a/Engine/Source/Runtime/Core/Private/ProfilingDebugging/CountersTrace.cpp
+++ b/Engine/Source/Runtime/Core/Private/ProfilingDebugging/CountersTrace.cpp
@@ -6,7 +6,9 @@
 
 #if COUNTERSTRACE_ENABLED
 
-UE_TRACE_CHANNEL_DEFINE(CountersChannel)
+UE_TRACE_CHANNEL_DEFINE(CountersChannel, "Enables the Counters Trace API, allows tracing of values (float, integer) over time. \
+Performance counters and statistics tracking. Data populates the Counters table in the Timing Insights profiler. \
+It also enables the CSV Profiler Trace.")
 
 UE_TRACE_EVENT_BEGIN(Counters, Spec, NoSync|Important)
 	UE_TRACE_EVENT_FIELD(uint16, Id)
diff --git a/Engine/Source/Runtime/Core/Private/ProfilingDebugging/CpuProfilerTrace.cpp b/Engine/Source/Runtime/Core/Private/ProfilingDebugging/CpuProfilerTrace.cpp
--- a/Engine/Source/Runtime/Core/Private/ProfilingDebugging/CpuProfilerTrace.cpp
+++ b/Engine/Source/Runtime/Core/Private/ProfilingDebugging/CpuProfilerTrace.cpp
@@ -21,7 +21,7 @@
 	#define CPUPROFILERTRACE_FILE_AND_LINE_ENABLED 1
 #endif
 
-UE_TRACE_CHANNEL_DEFINE(CpuChannel)
+UE_TRACE_CHANNEL_DEFINE(CpuChannel, "CPU profiler scoping events and timing data. Captures hierarchical timing information for CPU-bound operations on all threads.")
 
 UE_TRACE_EVENT_BEGIN(CpuProfiler, EventSpec, NoSync|Important)
 	UE_TRACE_EVENT_FIELD(uint32, Id)
diff --git a/Engine/Source/Runtime/Core/Private/ProfilingDebugging/IoStoreTrace.cpp b/Engine/Source/Runtime/Core/Private/ProfilingDebugging/IoStoreTrace.cpp
--- a/Engine/Source/Runtime/Core/Private/ProfilingDebugging/IoStoreTrace.cpp
+++ b/Engine/Source/Runtime/Core/Private/ProfilingDebugging/IoStoreTrace.cpp
@@ -15,7 +15,7 @@
 #endif
 
 
-UE_TRACE_CHANNEL_DEFINE(IoStoreChannel)
+UE_TRACE_CHANNEL_DEFINE(IoStoreChannel, "I/O dispatcher requests, backend operations, and chunk loading.")
 UE_TRACE_METADATA_EVENT_DEFINE(IoStoreTag)
 
 UE_TRACE_EVENT_BEGIN(IoStore, BackendName, NoSync)
diff --git a/Engine/Source/Runtime/Core/Private/ProfilingDebugging/MemoryTrace.cpp b/Engine/Source/Runtime/Core/Private/ProfilingDebugging/MemoryTrace.cpp
--- a/Engine/Source/Runtime/Core/Private/ProfilingDebugging/MemoryTrace.cpp
+++ b/Engine/Source/Runtime/Core/Private/ProfilingDebugging/MemoryTrace.cpp
@@ -13,7 +13,8 @@
 #include "Trace/Trace.h"
 
 #if UE_TRACE_ENABLED
-UE_TRACE_CHANNEL_DEFINE(MemAllocChannel, "Memory allocations", true)
+UE_TRACE_CHANNEL_DEFINE(MemAllocChannel, "Traces memory allocations and displays them in the Memory Insights Profiler. Optional callstack data is associated using Module + Callstack channels. \
+Optional asset metadata is associated using Metadata + AssetMetadata channels.", true)
 #endif
 
 #if UE_MEMORY_TRACE_ENABLED
diff --git a/Engine/Source/Runtime/Core/Private/ProfilingDebugging/MetadataTrace.cpp b/Engine/Source/Runtime/Core/Private/ProfilingDebugging/MetadataTrace.cpp
--- a/Engine/Source/Runtime/Core/Private/ProfilingDebugging/MetadataTrace.cpp
+++ b/Engine/Source/Runtime/Core/Private/ProfilingDebugging/MetadataTrace.cpp
@@ -5,7 +5,7 @@
 #if UE_TRACE_METADATA_ENABLED
 ////////////////////////////////////////////////////////////////////////////////////////////////////
 
-UE_TRACE_CHANNEL_DEFINE(MetadataChannel);
+UE_TRACE_CHANNEL_DEFINE(MetadataChannel, "Metadata stack management for trace context propagation. Associated with memory allocations.");
 
 UE_TRACE_EVENT_DEFINE(MetadataStack, ClearScope);
 UE_TRACE_EVENT_DEFINE(MetadataStack, SaveStack);
diff --git a/Engine/Source/Runtime/Core/Private/ProfilingDebugging/MiscTrace.cpp b/Engine/Source/Runtime/Core/Private/ProfilingDebugging/MiscTrace.cpp
--- a/Engine/Source/Runtime/Core/Private/ProfilingDebugging/MiscTrace.cpp
+++ b/Engine/Source/Runtime/Core/Private/ProfilingDebugging/MiscTrace.cpp
@@ -6,10 +6,10 @@
 #include "HAL/PlatformTLS.h"
 #include "HAL/PlatformTime.h"
 
-UE_TRACE_MINIMAL_CHANNEL(FrameChannel)
-UE_TRACE_CHANNEL_DEFINE(BookmarkChannel)
-UE_TRACE_CHANNEL(RegionChannel)
-UE_TRACE_CHANNEL(ScreenshotChannel)
+UE_TRACE_MINIMAL_CHANNEL(FrameChannel, "Traces the frame durations for different frame types. They can be viewed in the Timing Insights profiler.")
+UE_TRACE_CHANNEL_DEFINE(BookmarkChannel, "Low frequency markers signifying important transitions, providing a quick overview for users. Such as level loading or engine boot phases.")
+UE_TRACE_CHANNEL(RegionChannel, "Allows tracing of regions. Regions are thread-agnostic timespans in the trace. They can be visualized in the Timing Insights profiler.")
+UE_TRACE_CHANNEL(ScreenshotChannel, "Allows screenshots to be included in the trace. Screenshots can be visualized in the Timing Insights profiler.")
 
 UE_TRACE_EVENT_BEGIN(Misc, BookmarkSpec, NoSync|Important)
 	UE_TRACE_EVENT_FIELD(const void*, BookmarkPoint)
diff --git a/Engine/Source/Runtime/Core/Private/ProfilingDebugging/PlatformEvents.cpp b/Engine/Source/Runtime/Core/Private/ProfilingDebugging/PlatformEvents.cpp
--- a/Engine/Source/Runtime/Core/Private/ProfilingDebugging/PlatformEvents.cpp
+++ b/Engine/Source/Runtime/Core/Private/ProfilingDebugging/PlatformEvents.cpp
@@ -57,8 +57,8 @@ class FSSChannel : public UE::Trace::FChannel
 };
 #endif
 
-UE_TRACE_CHANNEL_CUSTOM_DEFINE(ContextSwitchChannel, FCSChannel, "CPU context switches", false)
-UE_TRACE_CHANNEL_CUSTOM_DEFINE(StackSamplingChannel, FSSChannel, "Stack sampling", false)
+UE_TRACE_CHANNEL_CUSTOM_DEFINE(ContextSwitchChannel, FCSChannel, "Trace CPU context switches. On Windows, game/editor runtime should be run as administrator", false)
+UE_TRACE_CHANNEL_CUSTOM_DEFINE(StackSamplingChannel, FSSChannel, "Trace stack sampling events. On Windows, game/editor runtime should be run as administrator", false)
 
 /////////////////////////////////////////////////////////////////////
 
diff --git a/Engine/Source/Runtime/Core/Private/Sanitizer/RaceDetector.cpp b/Engine/Source/Runtime/Core/Private/Sanitizer/RaceDetector.cpp
--- a/Engine/Source/Runtime/Core/Private/Sanitizer/RaceDetector.cpp
+++ b/Engine/Source/Runtime/Core/Private/Sanitizer/RaceDetector.cpp
@@ -33,7 +33,7 @@
 
 DEFINE_LOG_CATEGORY(LogRaceDetector);
 
-UE_TRACE_CHANNEL_DEFINE(RaceDetectorChannel);
+UE_TRACE_CHANNEL_DEFINE(RaceDetectorChannel, "Race condition detection for thread safety analysis. Part of the sanitizer infrastructure for detecting data races and concurrent memory access violations.");
 
 #define WITH_RACEDETECTOR_CHECK 0
 #define WITH_RACEDETECTOR_DEBUG 0
diff --git a/Engine/Source/Runtime/Core/Private/Serialization/LoadTimeTrace.cpp b/Engine/Source/Runtime/Core/Private/Serialization/LoadTimeTrace.cpp
--- a/Engine/Source/Runtime/Core/Private/Serialization/LoadTimeTrace.cpp
+++ b/Engine/Source/Runtime/Core/Private/Serialization/LoadTimeTrace.cpp
@@ -8,8 +8,8 @@
 #include "Misc/CString.h"
 #include "HAL/PlatformTLS.h"
 
-UE_TRACE_CHANNEL_DEFINE(LoadTimeChannel)
-UE_TRACE_CHANNEL_DEFINE(AssetLoadTimeChannel)
+UE_TRACE_CHANNEL_DEFINE(LoadTimeChannel, "Load time profiling for packages and request groups. Can be viewed in Asset Loading Insights.")
+UE_TRACE_CHANNEL_DEFINE(AssetLoadTimeChannel, "Asset-specific load time profiling and timing. This channel is only used as a toggle for CPU Profiler scoped events related to asset loading.")
 
 UE_TRACE_EVENT_BEGIN(LoadTime, BeginRequestGroup)
 	UE_TRACE_EVENT_FIELD(UE::Trace::WideString, Format)
diff --git a/Engine/Source/Runtime/Core/Private/Stats/ThreadIdleStats.cpp b/Engine/Source/Runtime/Core/Private/Stats/ThreadIdleStats.cpp
--- a/Engine/Source/Runtime/Core/Private/Stats/ThreadIdleStats.cpp
+++ b/Engine/Source/Runtime/Core/Private/Stats/ThreadIdleStats.cpp
@@ -4,7 +4,7 @@
 
 #if !(defined(DISABLE_THREAD_IDLE_STATS) && DISABLE_THREAD_IDLE_STATS)
 #if CPUPROFILERTRACE_ENABLED
-UE_TRACE_CHANNEL_DEFINE(ThreadIdleScopeChannel);
+UE_TRACE_CHANNEL_DEFINE(ThreadIdleScopeChannel, "Thread idle time tracking. Marks periods when threads are idle or waiting by adding cpu profiler events for them.");
 TRACE_CPUPROFILER_EVENT_DECLARE(ThreadIdleScopeTraceEventId);
 #endif
 
diff --git a/Engine/Source/Runtime/CoreUObject/Private/UObject/SavePackage/SavePackageUtilities.cpp b/Engine/Source/Runtime/CoreUObject/Private/UObject/SavePackage/SavePackageUtilities.cpp
--- a/Engine/Source/Runtime/CoreUObject/Private/UObject/SavePackage/SavePackageUtilities.cpp
+++ b/Engine/Source/Runtime/CoreUObject/Private/UObject/SavePackage/SavePackageUtilities.cpp
@@ -51,7 +51,7 @@
 #include "UObject/UnrealType.h"
 
 DEFINE_LOG_CATEGORY(LogSavePackage);
-UE_TRACE_CHANNEL_DEFINE(SaveTimeChannel);
+UE_TRACE_CHANNEL_DEFINE(SaveTimeChannel, "Adds extra cpu timing events that track serialization and package saving operations.");
 
 #if ENABLE_COOK_STATS
 #include "ProfilingDebugging/ScopedTimers.h"
diff --git a/Engine/Source/Runtime/Engine/Private/Animation/AnimTrace.cpp b/Engine/Source/Runtime/Engine/Private/Animation/AnimTrace.cpp
--- a/Engine/Source/Runtime/Engine/Private/Animation/AnimTrace.cpp
+++ b/Engine/Source/Runtime/Engine/Private/Animation/AnimTrace.cpp
@@ -19,7 +19,7 @@
 #include "UObject/UObjectAnnotation.h"
 #include "Engine/PoseWatch.h"
 
-UE_TRACE_CHANNEL_DEFINE(AnimationChannel);
+UE_TRACE_CHANNEL_DEFINE(AnimationChannel, "Animation tick records, skeletal mesh poses, and animation graph evaluation.");
 
 UE_TRACE_EVENT_BEGIN(Animation, TickRecord3)
 	UE_TRACE_EVENT_FIELD(uint64, Cycle)
diff --git a/Engine/Source/Runtime/MassEntity/Private/MassEntityTrace.cpp b/Engine/Source/Runtime/MassEntity/Private/MassEntityTrace.cpp
--- a/Engine/Source/Runtime/MassEntity/Private/MassEntityTrace.cpp
+++ b/Engine/Source/Runtime/MassEntity/Private/MassEntityTrace.cpp
@@ -14,7 +14,7 @@
 #include "Trace/Trace.inl"
 #include "TraceFilter.h"
 
-UE_TRACE_CHANNEL_DEFINE(MassChannel);
+UE_TRACE_CHANNEL_DEFINE(MassChannel, "Mass Entity system archetypes, queries, and chunk execution.");
 
 UE_TRACE_EVENT_BEGIN(MassTrace, MassPhaseBegin)
 	UE_TRACE_EVENT_FIELD(uint64, Cycle)
diff --git a/Engine/Source/Runtime/RHI/Private/GpuProfilerTrace.cpp b/Engine/Source/Runtime/RHI/Private/GpuProfilerTrace.cpp
--- a/Engine/Source/Runtime/RHI/Private/GpuProfilerTrace.cpp
+++ b/Engine/Source/Runtime/RHI/Private/GpuProfilerTrace.cpp
@@ -15,7 +15,7 @@
 
 #if UE_TRACE_GPU_PROFILER_ENABLED
 	UE_TRACE_CHANNEL_EXTERN(GpuChannel, RHI_API)
-	UE_TRACE_CHANNEL_DEFINE(GpuChannel)
+	UE_TRACE_CHANNEL_DEFINE(GpuChannel, "GPU profiler events including timing, breadcrumbs, and queue synchronization. Data can be viewed in the Timing Profiler.")
 #endif
 
 namespace UE::RHI::GPUProfiler
diff --git a/Engine/Source/Runtime/RHI/Private/RHICommandList.cpp b/Engine/Source/Runtime/RHI/Private/RHICommandList.cpp
--- a/Engine/Source/Runtime/RHI/Private/RHICommandList.cpp
+++ b/Engine/Source/Runtime/RHI/Private/RHICommandList.cpp
@@ -34,7 +34,7 @@ DECLARE_CYCLE_STAT(TEXT("All Command List Execute"), STAT_ImmedCmdListExecuteTim
 DECLARE_DWORD_COUNTER_STAT(TEXT("Immed. Command List memory"), STAT_ImmedCmdListMemory, STATGROUP_RHICMDLIST);
 DECLARE_DWORD_COUNTER_STAT(TEXT("Immed. Command count"), STAT_ImmedCmdListCount, STATGROUP_RHICMDLIST);
 
-UE_TRACE_CHANNEL_DEFINE(RHICommandsChannel);
+UE_TRACE_CHANNEL_DEFINE(RHICommandsChannel, "RHI command list execution and translation.");
 
 #if VALIDATE_UNIFORM_BUFFER_STATIC_BINDINGS
 static thread_local bool bScopedUniformBufferStaticBindingsRecursionGuard = false;
diff --git a/Engine/Source/Runtime/RenderCore/Private/RenderGraphTrace.cpp b/Engine/Source/Runtime/RenderCore/Private/RenderGraphTrace.cpp
--- a/Engine/Source/Runtime/RenderCore/Private/RenderGraphTrace.cpp
+++ b/Engine/Source/Runtime/RenderCore/Private/RenderGraphTrace.cpp
@@ -10,7 +10,7 @@
 
 #if RDG_ENABLE_TRACE
 
-UE_TRACE_CHANNEL_DEFINE(RDGChannel)
+UE_TRACE_CHANNEL_DEFINE(RDGChannel, "Render Dependency Graph passes, resources, and transient allocations.")
 
 UE_TRACE_EVENT_BEGIN(RDGTrace, GraphMessage)
 	UE_TRACE_EVENT_FIELD(UE::Trace::WideString, Name)
diff --git a/Engine/Source/Runtime/RenderCore/Private/RenderingThread.cpp b/Engine/Source/Runtime/RenderCore/Private/RenderingThread.cpp
--- a/Engine/Source/Runtime/RenderCore/Private/RenderingThread.cpp
+++ b/Engine/Source/Runtime/RenderCore/Private/RenderingThread.cpp
@@ -44,7 +44,7 @@
 FCoreRenderDelegates::FOnFlushRenderingCommandsStart FCoreRenderDelegates::OnFlushRenderingCommandsStart;
 FCoreRenderDelegates::FOnFlushRenderingCommandsEnd FCoreRenderDelegates::OnFlushRenderingCommandsEnd;
 
-UE_TRACE_CHANNEL_DEFINE(RenderCommandsChannel);
+UE_TRACE_CHANNEL_DEFINE(RenderCommandsChannel, "Render thread command execution and synchronization.");
 
 RENDERCORE_API bool GIsThreadedRendering = false;
 RENDERCORE_API bool GUseThreadedRendering = false;
