UnifiedError supports LexToString to and bool for error params.

#Synced-CL 52313638

diff --git a/Engine/Source/Runtime/Core/Public/Experimental/UnifiedError/UnifiedError.h b/Engine/Source/Runtime/Core/Public/Experimental/UnifiedError/UnifiedError.h
--- a/Engine/Source/Runtime/Core/Public/Experimental/UnifiedError/UnifiedError.h
+++ b/Engine/Source/Runtime/Core/Public/Experimental/UnifiedError/UnifiedError.h
@@ -646,6 +646,13 @@ namespace UE::UnifiedError
 		Writer.AddString(PropertyName, Value.ToString());
 	}
 
+	// Allow bool but not implicit conversions to bool
+	template<typename BoolT> requires UE::CDecaysTo<BoolT, bool>
+	inline void SerializeBasicToCb(FCbWriter& Writer, const FUtf8StringView& PropertyName, BoolT Value)
+	{
+		Writer.AddBool(PropertyName, Value);
+	}
+
 	
 
 } // namespace UnifiedError
@@ -662,11 +669,33 @@ CORE_API constexpr uint32 GetTypeIdHash(const TCHAR* TypeName)
 */
 
 
+/* 
+ * Note: HACK: Putting templated SerializeBasicToCb in this macro is a hack.
+ * SerializeBasicToCb must be declared after any LexToString() overload in order to find the overload.
+ * Putting it in the macro delays its declaration until after any possible required LexToString() definitions.
+ * A proper fix would be to split this header into two separate headers - one for public use of FError,
+ * and another for declaring errors.
+ */
 #define UE_DECLARE_ERROR_MODULE(DeclareApi, ModuleName) \
-	namespace UE::UnifiedError { namespace ModuleName {\
-		DeclareApi FAnsiStringView GetStaticModuleName(); \
-		DeclareApi int32 GetStaticModuleId(); \
-	}}
+	namespace UE::UnifiedError \
+	{ \
+		namespace ModuleName \
+		{ \
+			DeclareApi FAnsiStringView GetStaticModuleName(); \
+			DeclareApi int32 GetStaticModuleId(); \
+			/* Try to get strings for other types if possible */ \
+			template<typename T> requires (!UE::CIntegral<T> && !UE::CFloatingPoint<T> && !UE::CStringViewable<T>) \
+			inline void SerializeBasicToCb(FCbWriter& Writer, const FUtf8StringView& PropertyName, const T& Value) \
+			{ \
+				Writer.AddString(PropertyName, LexToString(Value)); \
+			} \
+			template<typename T> requires (UE::CIntegral<T> || UE::CFloatingPoint<T> || UE::CStringViewable<T>) \
+			inline void SerializeBasicToCb(FCbWriter& Writer, const FUtf8StringView& PropertyName, const T& Value) \
+			{ \
+				UE::UnifiedError::SerializeBasicToCb(Writer, PropertyName, Value); \
+			} \
+		} \
+	}
 
 #define UE_DEFINE_ERROR_MODULE(ModuleName) \
 	namespace UE::UnifiedError { namespace ModuleName {\
@@ -749,7 +778,7 @@ CORE_API constexpr uint32 GetTypeIdHash(const TCHAR* TypeName)
 		Writer.BeginObject();\
 		Writer.AddString(ANSITEXTVIEW("$type"), TErrorStructFeatures<UE::UnifiedError::ModuleName::F##ErrorName>::GetErrorContextTypeNameAsString());\
 		UE::SerializeLogFormat(Writer, FormatString); \
-		UE::UnifiedError::SerializeBasicToCb(Writer, ANSITEXTVIEW(#ParamOneName), Context.ParamOneName);\
+		UE::UnifiedError::ModuleName::SerializeBasicToCb(Writer, ANSITEXTVIEW(#ParamOneName), Context.ParamOneName);\
 		Writer.EndObject();\
 	}\
 	namespace UE::UnifiedError { namespace ModuleName { \
@@ -778,8 +807,8 @@ CORE_API constexpr uint32 GetTypeIdHash(const TCHAR* TypeName)
 		Writer.BeginObject();\
 		Writer.AddString(ANSITEXTVIEW("$type"), TErrorStructFeatures<UE::UnifiedError::ModuleName::F##ErrorName>::GetErrorContextTypeNameAsString());\
 		UE::SerializeLogFormat(Writer, FormatString); \
-		UE::UnifiedError::SerializeBasicToCb(Writer, ANSITEXTVIEW(#ParamOneName), Context.ParamOneName);\
-		UE::UnifiedError::SerializeBasicToCb(Writer, ANSITEXTVIEW(#ParamTwoName), Context.ParamTwoName);\
+		UE::UnifiedError::ModuleName::SerializeBasicToCb(Writer, ANSITEXTVIEW(#ParamOneName), Context.ParamOneName);\
+		UE::UnifiedError::ModuleName::SerializeBasicToCb(Writer, ANSITEXTVIEW(#ParamTwoName), Context.ParamTwoName);\
 		Writer.EndObject();\
 	}\
 	namespace UE::UnifiedError { namespace ModuleName { \
