Prevent LexToString(bool) from hijacking types that implicitly cast to bool.  This makes missing or unresolvable LexToString overloads fail at compile time instead of being hidden by a bool conversion.

#Synced-CL 52313638

diff --git a/Engine/Plugins/Runtime/Metasound/Source/MetasoundFrontend/Private/MetasoundTrigger.cpp b/Engine/Plugins/Runtime/Metasound/Source/MetasoundFrontend/Private/MetasoundTrigger.cpp
--- a/Engine/Plugins/Runtime/Metasound/Source/MetasoundFrontend/Private/MetasoundTrigger.cpp
+++ b/Engine/Plugins/Runtime/Metasound/Source/MetasoundFrontend/Private/MetasoundTrigger.cpp
@@ -159,3 +159,13 @@ namespace Metasound
 		LastTriggerIndexInBlock = Algo::LowerBound(TriggeredFrames, NumFramesPerBlock);
 	}
 }
+
+const TCHAR* LexToString(const Metasound::FTrigger& Trigger)
+{
+	if (Trigger)
+	{
+		return TEXT("Metasound::FTrigger_True");
+	}
+
+	return TEXT("Metasound::FTrigger_False");
+}
diff --git a/Engine/Plugins/Runtime/Metasound/Source/MetasoundFrontend/Public/MetasoundTrigger.h b/Engine/Plugins/Runtime/Metasound/Source/MetasoundFrontend/Public/MetasoundTrigger.h
--- a/Engine/Plugins/Runtime/Metasound/Source/MetasoundFrontend/Public/MetasoundTrigger.h
+++ b/Engine/Plugins/Runtime/Metasound/Source/MetasoundFrontend/Public/MetasoundTrigger.h
@@ -204,4 +204,6 @@ namespace Metasound
 	DECLARE_METASOUND_DATA_REFERENCE_TYPES(FTrigger, METASOUNDFRONTEND_API, FTriggerTypeInfo, FTriggerReadRef, FTriggerWriteRef);
 }
 
+UE_API const TCHAR* LexToString(const Metasound::FTrigger& Trigger);
+
 #undef UE_API
diff --git a/Engine/Source/Runtime/Core/Public/Containers/UnrealString.h b/Engine/Source/Runtime/Core/Public/Containers/UnrealString.h
--- a/Engine/Source/Runtime/Core/Public/Containers/UnrealString.h
+++ b/Engine/Source/Runtime/Core/Public/Containers/UnrealString.h
@@ -5,6 +5,8 @@
 
 #pragma once
 
+#include "Concepts/DecaysTo.h"
+
 // Include UnrealString.h.inl's includes before defining the macros, in case the macros 'poison' other headers or there are re-entrant includes.
 #include "Containers/UnrealStringIncludes.h.inl" // IWYU pragma: export
 
@@ -146,13 +148,13 @@ inline void ByteToHex(uint8 In, FString& Result)
 	return 0;
 }
 
-/** Convert numeric types to a string */
+/** Convert numeric types to a string. Excludes bool to avoid hijacking through implicit conversion. */
 template <
 	typename StringType = FString,
 	typename T
 	UE_REQUIRES(std::is_arithmetic_v<T>)
 >
-[[nodiscard]] StringType LexToString(const T& Value)
+[[nodiscard]] StringType LexToString(const T& Value) requires (!UE::CDecaysTo<T, bool>)
 {
 	// std::remove_cv_t to remove potential volatile decorations. Removing const is pointless, but harmless because it's specified in the param declaration.
 	return StringType::Printf(TFormatSpecifier<std::remove_cv_t<T>>::template GetFormatSpecifier<typename StringType::FmtCharType>(), Value);
@@ -168,8 +170,9 @@ template <
 	return StringType(Ptr);
 }
 
-template <typename StringType = FString>
-[[nodiscard]] inline StringType LexToString(bool Value)
+/** Convert bool to a string. Do not allow implicit conversion to avoid highjacking other types. */
+template <typename StringType = FString, typename BoolT>
+[[nodiscard]] inline StringType LexToString(BoolT Value) requires UE::CDecaysTo<BoolT, bool>
 {
 	using ElementType = typename StringType::ElementType;
 	return Value ? CHARTEXT(ElementType, "true") : CHARTEXT(ElementType, "false");
