Add additional warnings to the set disabled for third party code.
#rnx
#Synced-CL 52374809

diff --git a/Engine/Source/Runtime/Core/Public/Microsoft/MSVCPlatformCompilerPreSetup.h b/Engine/Source/Runtime/Core/Public/Microsoft/MSVCPlatformCompilerPreSetup.h
--- a/Engine/Source/Runtime/Core/Public/Microsoft/MSVCPlatformCompilerPreSetup.h
+++ b/Engine/Source/Runtime/Core/Public/Microsoft/MSVCPlatformCompilerPreSetup.h
@@ -234,10 +234,13 @@
 	#ifndef UE_COMPILER_THIRD_PARTY_INCLUDES_START
 		#define UE_COMPILER_THIRD_PARTY_INCLUDES_START \
 			__pragma(warning(push)) \
+			__pragma(warning(disable: 4100))  /* 'identifier': unreferenced formal parameter. */ \
 			__pragma(warning(disable: 4125))  /* decimal digit terminates octal escape sequence. */ \
+			__pragma(warning(disable: 4267))  /* 'var': conversion from 'size_t' to 'type', possible loss of data. */ \
 			__pragma(warning(disable: 4510))  /* '<class>': default constructor could not be generated. */ \
 			__pragma(warning(disable: 4610))  /* object '<class>' can never be instantiated - user-defined constructor required. */ \
 			__pragma(warning(disable: 4800))  /* Implicit conversion from '<type>' to bool. Possible information loss. */ \
+			__pragma(warning(disable: 4834))  /* discarding return value of function with 'nodiscard' attribute. */ \
 			__pragma(warning(disable: 4946))  /* reinterpret_cast used between related classes: '<class1>' and '<class2>' */ \
 			__pragma(warning(disable: 4996))  /* '<obj>' was declared deprecated. */ \
 			__pragma(warning(disable: 6011))  /* Dereferencing NULL pointer '<ptr>'. */ \
