Fix non-unity issue introduced by CL# 52392201
#rnx
#Synced-CL 52414109

diff --git a/Engine/Source/Runtime/Core/Public/Templates/Overload.h b/Engine/Source/Runtime/Core/Public/Templates/Overload.h
--- a/Engine/Source/Runtime/Core/Public/Templates/Overload.h
+++ b/Engine/Source/Runtime/Core/Public/Templates/Overload.h
@@ -54,6 +54,7 @@
 
 #include <type_traits>
 
+#include "HAL/Platform.h"
 #include "Templates/Projection.h"
 
 namespace UE::Core::Private
