Updated EMemoryUnit to offer explicit Auto_IEC and Auto_SI.

#rnx

#Synced-CL 52109349

diff --git a/Engine/Source/Runtime/Core/Private/Internationalization/Text.cpp b/Engine/Source/Runtime/Core/Private/Internationalization/Text.cpp
--- a/Engine/Source/Runtime/Core/Private/Internationalization/Text.cpp
+++ b/Engine/Source/Runtime/Core/Private/Internationalization/Text.cpp
@@ -836,9 +836,9 @@ FText FText::AsMemory(uint64 NumBytes, EMemoryUnitStandard UnitStandard)
 
 FText FText::AsMemory(uint64 NumBytes, EMemoryUnit Unit, const FNumberFormattingOptions* const Options, const FCulturePtr& TargetCulture)
 {
-	if (Unit == EMemoryUnit::Auto)
+	if (Unit == EMemoryUnit::Auto_IEC || Unit == EMemoryUnit::Auto_SI)
 	{
-		return FText::AsMemory(NumBytes, Options, TargetCulture);
+		return FText::AsMemory(NumBytes, Options, TargetCulture, Unit == EMemoryUnit::Auto_IEC ? EMemoryUnitStandard::IEC : EMemoryUnitStandard::SI);
 	}
 
 	checkf(FInternationalization::Get().IsInitialized() == true, TEXT("FInternationalization is not initialized. An FText formatting method was likely used in static object initialization - this is not supported."));
@@ -851,10 +851,15 @@ FText FText::AsMemory(uint64 NumBytes, EMemoryUnit Unit, const FNumberFormatting
 		return FText::Format(NSLOCTEXT("Internationalization", "ComputerMemoryFormatting", "{Number} {Unit}"), Args);
 	}
 
+	static_assert((int32)EMemoryUnit::Unit_EiB - (int32)EMemoryUnit::Unit_KiB == 5, "IEC units must be contiguous in ascending order from Unit_KiB to Unit_EiB");
+	static_assert((int32)EMemoryUnit::Unit_EB - (int32)EMemoryUnit::Unit_KB == 5, "SI units must be contiguous in ascending order from Unit_KB to Unit_EB");
+	static_assert((int32)EMemoryUnit::Auto_SI > (int32)EMemoryUnit::Unit_EiB, "Auto_SI must be ordered after all IEC units for the bIEC range check to be correct");
+	static_assert((int32)EMemoryUnit::Unit_KB > (int32)EMemoryUnit::Auto_SI, "SI units must be ordered after Auto_SI for the offset calculation to be correct");
+
 	static const TCHAR* IECPrefixes = TEXT("KMGTPE");
-	static const TCHAR* SIPrefixes  = TEXT("kMGTPE");
+	static const TCHAR* SIPrefixes = TEXT("kMGTPE");
 
-	const bool bIEC = (Unit <= EMemoryUnit::Unit_EiB);
+	const bool bIEC = (Unit >= EMemoryUnit::Unit_KiB && Unit <= EMemoryUnit::Unit_EiB);
 	const int32 PrefixIndex = bIEC
 		? (int32)Unit - (int32)EMemoryUnit::Unit_KiB
 		: (int32)Unit - (int32)EMemoryUnit::Unit_KB;
diff --git a/Engine/Source/Runtime/Core/Public/Internationalization/Text.h b/Engine/Source/Runtime/Core/Public/Internationalization/Text.h
--- a/Engine/Source/Runtime/Core/Public/Internationalization/Text.h
+++ b/Engine/Source/Runtime/Core/Public/Internationalization/Text.h
@@ -199,11 +199,17 @@ enum EMemoryUnitStandard : int
 /** Specific memory unit to use when formatting memory sizes. IEC units are 1024-based (KiB, MiB, ...), SI units are 1000-based (KB, MB, ...). */
 enum class EMemoryUnit : int
 {
-	// Auto-select based on magnitude (default behavior)
-	Auto,
 	Unit_Bytes,
+
+	// Auto-select based on magnitude using IEC
+	Auto_IEC,
+
 	// IEC (1024-based)
 	Unit_KiB, Unit_MiB, Unit_GiB, Unit_TiB, Unit_PiB, Unit_EiB,
+
+	// Auto-select based on magnitude using SI
+	Auto_SI,
+
 	// SI (1000-based)
 	Unit_KB, Unit_MB, Unit_GB, Unit_TB, Unit_PB, Unit_EB
 };
@@ -490,7 +496,7 @@ public:
 
 	/**
 	 * Generate an FText that represents the passed number as a memory size in the current culture, expressed in a specific unit.
-	 * Use EMemoryUnit::Auto to preserve the default magnitude-based selection behavior.
+	 * Use EMemoryUnit::Auto_IEC or EMemoryUnit::Auto_SI to preserve the default magnitude-based selection behavior.
 	 */
 	[[nodiscard]] static CORE_API FText AsMemory(uint64 NumBytes, EMemoryUnit Unit, const FNumberFormattingOptions* const Options = nullptr, const FCulturePtr& TargetCulture = nullptr);
 
diff --git a/Engine/Source/Runtime/Core/Tests/Internationalization/TextTest.cpp b/Engine/Source/Runtime/Core/Tests/Internationalization/TextTest.cpp
--- a/Engine/Source/Runtime/Core/Tests/Internationalization/TextTest.cpp
+++ b/Engine/Source/Runtime/Core/Tests/Internationalization/TextTest.cpp
@@ -389,34 +389,88 @@ public:
 					.SetMinimumFractionalDigits(0)
 					.SetMaximumFractionalDigits(3);
 
-				// Auto delegates to existing IEC behavior
-				TEST_UNIT(0, EMemoryUnit::Auto, "0 B");
-				TEST_UNIT(500, EMemoryUnit::Auto, "500 B");
-				TEST_UNIT(1048576, EMemoryUnit::Auto, "1 MiB");
-
-				// Explicit Bytes
+				// --- Auto_IEC (These will map to FText::AsMemory() not accepting EMemoryUnit) ---
+				// Bytes range: 0, 1, and just below the 1 KiB boundary
+				TEST_UNIT(0, EMemoryUnit::Auto_IEC, "0 B");
+				TEST_UNIT(1, EMemoryUnit::Auto_IEC, "1 B");
+				TEST_UNIT(1023, EMemoryUnit::Auto_IEC, "0.999 KiB");
+				
+				// Exact powers of 1024 (KiB through EiB)
+				TEST_UNIT(1024, EMemoryUnit::Auto_IEC, "1 KiB");
+				TEST_UNIT(1048576, EMemoryUnit::Auto_IEC, "1 MiB");
+				TEST_UNIT(1073741824, EMemoryUnit::Auto_IEC, "1 GiB");
+				TEST_UNIT(1099511627776, EMemoryUnit::Auto_IEC, "1 TiB");
+				TEST_UNIT(1125899906842624, EMemoryUnit::Auto_IEC, "1 PiB");
+				TEST_UNIT(1152921504606846976, EMemoryUnit::Auto_IEC, "1 EiB");
+				
+				// Just below the next boundary
+				TEST_UNIT(1048575, EMemoryUnit::Auto_IEC, "0.999 MiB");
+				TEST_UNIT(1073741823, EMemoryUnit::Auto_IEC, "0.999 GiB");
+				TEST_UNIT(1099511627775, EMemoryUnit::Auto_IEC, "0.999 TiB");
+				TEST_UNIT(1125899906842623, EMemoryUnit::Auto_IEC, "0.999 PiB");
+				TEST_UNIT(1152921504606846975, EMemoryUnit::Auto_IEC, "0.999 EiB");
+
+				// MAX uint64
+				TEST_UNIT(18446744073709551615ULL, EMemoryUnit::Auto_IEC, "15.999 EiB");
+
+				// --- Auto_SI (These will map to FText::AsMemory() not accepting EMemoryUnit) ---
+				// Bytes range: 0, 1, and just below the 1 kB boundary
+				TEST_UNIT(0, EMemoryUnit::Auto_SI, "0 B");
+				TEST_UNIT(1, EMemoryUnit::Auto_SI, "1 B");
+				TEST_UNIT(999, EMemoryUnit::Auto_SI, "999 B");
+
+				// Exact powers of 1000 (kB through EB)
+				TEST_UNIT(1000, EMemoryUnit::Auto_SI, "1 kB");
+				TEST_UNIT(1000000, EMemoryUnit::Auto_SI, "1 MB");
+				TEST_UNIT(1000000000, EMemoryUnit::Auto_SI, "1 GB");
+				TEST_UNIT(1000000000000, EMemoryUnit::Auto_SI, "1 TB");
+				TEST_UNIT(1000000000000000, EMemoryUnit::Auto_SI, "1 PB");
+				TEST_UNIT(1000000000000000000, EMemoryUnit::Auto_SI, "1 EB");
+
+				// Just below the next boundary
+				TEST_UNIT(999999, EMemoryUnit::Auto_SI, "999.999 kB");
+				TEST_UNIT(999999999, EMemoryUnit::Auto_SI, "999.999 MB");
+				TEST_UNIT(999999999999, EMemoryUnit::Auto_SI, "999.999 GB");
+				TEST_UNIT(999999999999999, EMemoryUnit::Auto_SI, "999.999 TB");
+				TEST_UNIT(999999999999999999, EMemoryUnit::Auto_SI, "999.999 PB");
+
+				// MAX uint64
+				TEST_UNIT(18446744073709551615ULL, EMemoryUnit::Auto_SI, "18.446 EB");
+
+				// --- Unit_Bytes ---
+				// Always formats as bytes regardless of magnitude
 				TEST_UNIT(0, EMemoryUnit::Unit_Bytes, "0 B");
 				TEST_UNIT(1, EMemoryUnit::Unit_Bytes, "1 B");
 				TEST_UNIT(999, EMemoryUnit::Unit_Bytes, "999 B");
+				TEST_UNIT(1024, EMemoryUnit::Unit_Bytes, "1,024 B");
 
-				// IEC units - exact powers of 1024
+				// --- Explicit IEC units ---
+				// Exact powers
 				TEST_UNIT(1024, EMemoryUnit::Unit_KiB, "1 KiB");
 				TEST_UNIT(1048576, EMemoryUnit::Unit_MiB, "1 MiB");
 				TEST_UNIT(1073741824, EMemoryUnit::Unit_GiB, "1 GiB");
 				TEST_UNIT(1099511627776, EMemoryUnit::Unit_TiB, "1 TiB");
 				TEST_UNIT(1125899906842624, EMemoryUnit::Unit_PiB, "1 PiB");
 				TEST_UNIT(1152921504606846976, EMemoryUnit::Unit_EiB, "1 EiB");
+				// Fractional values
+				TEST_UNIT(512, EMemoryUnit::Unit_KiB, "0.5 KiB");
+				TEST_UNIT(1536, EMemoryUnit::Unit_KiB, "1.5 KiB");
+				// MAX uint64 under explicit unit (float division: 18446744073709551615 / 1024^6 = ~16)
+				TEST_UNIT(18446744073709551615ULL, EMemoryUnit::Unit_EiB, "16 EiB");
 
-				// SI units - exact powers of 1000
+				// --- Explicit SI units ---
+				// Exact powers
 				TEST_UNIT(1000, EMemoryUnit::Unit_KB, "1 kB");
 				TEST_UNIT(1000000, EMemoryUnit::Unit_MB, "1 MB");
 				TEST_UNIT(1000000000, EMemoryUnit::Unit_GB, "1 GB");
 				TEST_UNIT(1000000000000, EMemoryUnit::Unit_TB, "1 TB");
 				TEST_UNIT(1000000000000000, EMemoryUnit::Unit_PB, "1 PB");
-				TEST_UNIT(1000000000000000000, EMemoryUnit::Unit_EB,  "1 EB");
-
-				// max uint64: float-based division yields 16 EiB (vs 15.999 from the integer-division path in the old overload)
-				TEST_UNIT(18446744073709551615ULL, EMemoryUnit::Unit_EiB, "16 EiB");
+				TEST_UNIT(1000000000000000000, EMemoryUnit::Unit_EB, "1 EB");
+				// Fractional values
+				TEST_UNIT(500, EMemoryUnit::Unit_KB, "0.5 kB");
+				TEST_UNIT(1500, EMemoryUnit::Unit_KB, "1.5 kB");
+				// MAX uint64 under explicit unit (float division: 18446744073709551615 / 1000^6 = ~18.447)
+				TEST_UNIT(18446744073709551615ULL, EMemoryUnit::Unit_EB, "18.447 EB");
 			}
 #undef TEST_UNIT
 
