LowLevelMemTracker, AutoRTFM: Fix all different types of FLLMScope structs to follow the pattern in FLLMScope: run Init and Destruct in the open, push/pop an AbortHandler. Standardize FLLMScope and all structs to have a DestructInTheOpen handler that skips the declaration of run in the open and the pop of the abort handler, since those are unnecessary when called from the abort handler.
#rnx
#Synced-CL 52434073

diff --git a/Engine/Source/Programs/AutoRTFMTests/Private/UECoreTests.cpp b/Engine/Source/Programs/AutoRTFMTests/Private/UECoreTests.cpp
--- a/Engine/Source/Programs/AutoRTFMTests/Private/UECoreTests.cpp
+++ b/Engine/Source/Programs/AutoRTFMTests/Private/UECoreTests.cpp
@@ -12,6 +12,7 @@
 #include "Experimental/ConcurrentLinearAllocator.h"
 #include "Experimental/UnifiedError/CoreErrorTypes.h"
 #include "HAL/IConsoleManager.h"
+#include "HAL/LowLevelMemTracker.h"
 #include "HAL/MallocLeakDetection.h"
 #include "HAL/ThreadHeartBeat.h"
 #include "HAL/ThreadSingleton.h"
@@ -557,12 +558,12 @@ TEST_CASE("UECore.FText")
 	{
 		SECTION("With Abort")
 		{
-			AutoRTFM::Testing::Commit([&]
+			AutoRTFM::Testing::Abort([&]
 			{
 				Text = FText::Format(NSLOCTEXT("Cat", "Dog", "Fish[{0}]"), uint64(255));
+				AutoRTFM::AbortTransaction();
 			});
-			REQUIRE(!Text.IsEmpty());
-			REQUIRE(Text.ToString() == TEXT("Fish[255]"));
+			REQUIRE(Text.IsEmpty());
 		}
 
 		SECTION("With Commit")
@@ -2354,3 +2355,240 @@ TEST_CASE("UECore.FSharedString")
 		REQUIRE(Compare == Other);
 	}
 }
+
+TEST_CASE("UECore.LowLevelMemTracker.Scopes")
+{
+	FName StartScopeName(TEXTVIEW("AutoRTFMTesting/StartScope"));
+	FName ApplyScopeName(TEXTVIEW("AutoRTFMTesting/ApplyScope"));
+	int64 StartSize = 0;
+
+	struct AUTORTFM_INFER FDeleterCallFree
+	{
+		void operator()(int32* Ptr) const
+		{
+			FMemory::Free(Ptr);
+		}
+	};
+
+	auto FetchStartSize = [&StartSize, ApplyScopeName]()
+		{
+			StartSize = 0;
+#if ENABLE_LOW_LEVEL_MEM_TRACKER
+			if (FLowLevelMemTracker::IsEnabled())
+			{
+				FLowLevelMemTracker& LLM = FLowLevelMemTracker::Get();
+				LLM.Tick();
+				StartSize = LLM.GetTagAmountForTracker(ELLMTracker::Default, ApplyScopeName, ELLMTagSet::None);
+			}
+#endif
+		};
+	auto RequireEndSize = [&StartSize, ApplyScopeName](int64 DeltaSize)
+		{
+#if ENABLE_LOW_LEVEL_MEM_TRACKER
+			if (FLowLevelMemTracker::IsEnabled())
+			{
+				FLowLevelMemTracker& LLM = FLowLevelMemTracker::Get();
+				LLM.Tick();
+				int64 EndSize = LLM.GetTagAmountForTracker(ELLMTracker::Default, ApplyScopeName, ELLMTagSet::None);
+				REQUIRE(EndSize == StartSize + DeltaSize);
+			}
+#endif
+		};
+
+	// Set the default scope that should be used when allocating outside of LLM_SCOPE et al declared in the tests.
+	LLM_SCOPE(StartScopeName);
+
+	SECTION("LLM_SCOPE")
+	{
+		SECTION("Commit")
+		{
+			TUniquePtr<int32, FDeleterCallFree> Allocation;
+			TUniquePtr<int32, FDeleterCallFree> Allocation2;
+			int64 DeltaSize = sizeof(int32);
+			FetchStartSize();
+
+			AutoRTFM::Testing::Commit([&]
+				{
+					LLM_SCOPE(ApplyScopeName);
+					Allocation.Reset(static_cast<int32*>(FMemory::Malloc(DeltaSize, alignof(int32))));
+				});
+			REQUIRE(Allocation);
+
+			RequireEndSize(DeltaSize);
+			// ApplyScope should no longer be applied and should not record this second DeltaSize.
+			Allocation2.Reset(static_cast<int32*>(FMemory::Malloc(DeltaSize, alignof(int32))));
+			RequireEndSize(DeltaSize);
+		}
+		SECTION("Abort")
+		{
+			int64 DeltaSize = sizeof(int32);
+			TUniquePtr<int32, FDeleterCallFree> Allocation2;
+			FetchStartSize();
+
+			AutoRTFM::Testing::Abort([&]
+				{
+					LLM_SCOPE(ApplyScopeName);
+					int32* Allocation = static_cast<int32*>(FMemory::Malloc(DeltaSize, alignof(int32)));
+					AutoRTFM::AbortTransaction();
+				});
+
+			// Abort calls Free and undoes the allocation, so DeltaSize should be 0.
+			RequireEndSize(0);
+			// ApplyScope should no longer be applied and should not record this second DeltaSize.
+			Allocation2.Reset(static_cast<int32*>(FMemory::Malloc(DeltaSize, alignof(int32))));
+			RequireEndSize(0);
+		}
+	}
+
+	SECTION("LLM_SCOPE_DYNAMIC")
+	{
+		SECTION("Commit")
+		{
+			int64 DeltaSize = sizeof(int32);
+			TUniquePtr<int32, FDeleterCallFree> Allocation;
+			TUniquePtr<int32, FDeleterCallFree> Allocation2;
+			FetchStartSize();
+
+			AutoRTFM::Testing::Commit([&]
+				{
+					FString LLMStatNameString(TEXT("AutoRTFMTesting"));
+					LLM_SCOPE_DYNAMIC(ApplyScopeName, ELLMTracker::Default, ELLMTagSet::None, FLLMDynamicTagConstructorStatString(LLMStatNameString));
+					Allocation.Reset(static_cast<int32*>(FMemory::Malloc(DeltaSize, alignof(int32))));
+				});
+			REQUIRE(Allocation);
+
+			RequireEndSize(DeltaSize);
+			// ApplyScope should no longer be applied and should not record this second DeltaSize.
+			Allocation2.Reset(static_cast<int32*>(FMemory::Malloc(DeltaSize, alignof(int32))));
+			RequireEndSize(DeltaSize);
+		}
+		SECTION("Abort")
+		{
+			int64 DeltaSize = sizeof(int32);
+			TUniquePtr<int32, FDeleterCallFree> Allocation2;
+			FetchStartSize();
+
+			AutoRTFM::Testing::Abort([&]
+				{
+					FString LLMStatNameString(TEXT("AutoRTFMTesting"));
+					LLM_SCOPE_DYNAMIC(ApplyScopeName, ELLMTracker::Default, ELLMTagSet::None, FLLMDynamicTagConstructorStatString(LLMStatNameString));
+					int32* Allocation = static_cast<int32*>(FMemory::Malloc(DeltaSize, alignof(int32)));
+					AutoRTFM::AbortTransaction();
+				});
+
+			// Abort calls Free and undoes the allocation, so DeltaSize should be 0.
+			RequireEndSize(0);
+			// ApplyScope should no longer be applied and should not record this second DeltaSize.
+			Allocation2.Reset(static_cast<int32*>(FMemory::Malloc(DeltaSize, alignof(int32))));
+			RequireEndSize(0);
+		}
+	}
+
+	SECTION("LLM_SCOPED_PAUSE_TRACKING")
+	{
+		SECTION("Commit")
+		{
+			int64 DeltaSize = sizeof(int32);
+			TUniquePtr<int32, FDeleterCallFree> Allocation;
+			TUniquePtr<int32, FDeleterCallFree> Allocation2;
+			FetchStartSize();
+
+			// Activate the Apply scope so we can verify that PAUSE_TRACKING blocks/unblocks it at the proper times.
+			LLM_SCOPE(ApplyScopeName);
+			AutoRTFM::Testing::Commit([&]
+				{
+					LLM_SCOPED_PAUSE_TRACKING(ELLMAllocType::FMalloc);
+					Allocation.Reset(static_cast<int32*>(FMemory::Malloc(DeltaSize, alignof(int32))));
+				});
+			REQUIRE(Allocation);
+
+			// PauseTracking means that allocations are not recorded.
+			RequireEndSize(0);
+			Allocation.Reset();
+			// Allocations made during PauseTracking are not recorded when they free.
+			RequireEndSize(0);
+
+			// PauseScope should no longer be applied and LLM should record this second DeltaSize.
+			Allocation2.Reset(static_cast<int32*>(FMemory::Malloc(DeltaSize, alignof(int32))));
+			RequireEndSize(DeltaSize);
+		}
+
+		SECTION("Abort")
+		{
+			int64 DeltaSize = sizeof(int32);
+			TUniquePtr<int32, FDeleterCallFree> Allocation2;
+			FetchStartSize();
+
+			LLM_SCOPE(ApplyScopeName);
+			AutoRTFM::Testing::Abort([&]
+				{
+					LLM_SCOPED_PAUSE_TRACKING(ELLMAllocType::System);
+					int32* Allocation = static_cast<int32*>(FMemory::Malloc(DeltaSize, alignof(int32)));
+					AutoRTFM::AbortTransaction();
+				});
+
+			// Abort calls Free and undoes the allocation. Also Pause prevented it from being reported.
+			RequireEndSize(0);
+
+			// PauseScope should no longer be applied and LLM should record this second DeltaSize.
+			Allocation2.Reset(static_cast<int32*>(FMemory::Malloc(DeltaSize, alignof(int32))));
+			RequireEndSize(DeltaSize);
+		}
+	}
+
+	SECTION("LLM_REALLOC_SCOPE")
+	{
+		SECTION("Commit")
+		{
+			int64 DeltaSize = 2 * sizeof(int32);
+			TUniquePtr<int32, FDeleterCallFree> Allocation;
+			TUniquePtr<int32, FDeleterCallFree> Allocation2;
+			FetchStartSize();
+
+			AutoRTFM::Testing::Commit([&]
+				{
+					{
+						LLM_SCOPE(ApplyScopeName);
+						Allocation.Reset(static_cast<int32*>(FMemory::Malloc(sizeof(int32), alignof(int32))));
+					}
+
+					int32* OldAllocation = Allocation.Release();
+					LLM_REALLOC_SCOPE(OldAllocation);
+					Allocation.Reset(static_cast<int32*>(FMemory::Realloc(OldAllocation, DeltaSize, alignof(int32))));
+				});
+			REQUIRE(Allocation);
+
+			// The second allocation should be the only one recorded since we called Realloc
+			RequireEndSize(DeltaSize);
+			// ReallocScope should no longer be applied and should not record this second DeltaSize.
+			Allocation2.Reset(static_cast<int32*>(FMemory::Malloc(DeltaSize, alignof(int32))));
+			RequireEndSize(DeltaSize);
+		}
+		SECTION("Abort")
+		{
+			int64 DeltaSize = 2 * sizeof(int32);
+			TUniquePtr<int32, FDeleterCallFree> Allocation2;
+			FetchStartSize();
+
+			AutoRTFM::Testing::Abort([&]
+				{
+					int32* Allocation;
+					{
+						LLM_SCOPE(ApplyScopeName);
+						Allocation = static_cast<int32*>(FMemory::Malloc(sizeof(int32), alignof(int32)));
+					}
+
+					int32* OldAllocation = Allocation;
+					LLM_REALLOC_SCOPE(OldAllocation);
+					Allocation = static_cast<int32*>(FMemory::Realloc(OldAllocation, DeltaSize, alignof(int32)));
+					AutoRTFM::AbortTransaction();
+				});
+
+			// Abort calls Free and undoes the allocation, so DeltaSize should be 0.
+			RequireEndSize(0);
+			// ReallocScope should no longer be applied and should not record this second DeltaSize.
+			Allocation2.Reset(static_cast<int32*>(FMemory::Malloc(DeltaSize, alignof(int32))));
+			RequireEndSize(0);
+		}
+	}
+}
diff --git a/Engine/Source/Runtime/Core/Private/HAL/LLM/LLM.cpp b/Engine/Source/Runtime/Core/Private/HAL/LLM/LLM.cpp
--- a/Engine/Source/Runtime/Core/Private/HAL/LLM/LLM.cpp
+++ b/Engine/Source/Runtime/Core/Private/HAL/LLM/LLM.cpp
@@ -3700,77 +3700,80 @@ FAllocationGroup* FLLMGlobals::FindAllocationGroupForCompressedAllocInfo(ELLMTag
 
 FLLMScope::FLLMScope(ELLMTag TagEnum, bool bInIsStatTag, ELLMTagSet InTagSet, ELLMTracker InTracker, bool bOverride)
 {
-	if (!FLowLevelMemTracker::IsEnabled())
-	{
-		return;
-	}
-
-	// We run the init in the open, because we want to track LLM even in transactions.
-	UE_AUTORTFM_OPEN
-	{
-		// This scope does not support TagSets (use FLLMScopeDynamic instead), except for the special case of
-		// anonymized allocations, which are indicated with a TagSet and ELLMTag::EngineMisc
-		LLMCheck((!bInIsStatTag && InTagSet == ELLMTagSet::None) || (TagEnum == ELLMTag::EngineMisc));
-		// ELLMTag::FMalloc is a special tag expected to only be used by the Platform Tracker (see header where defined).
-		LLMCheck((TagEnum != ELLMTag::FMalloc) || (ELLMTracker::Platform == InTracker));
+	using namespace UE::LLMPrivate;
 
-		UE::LLMPrivate::FLLMGlobals& LLMRef = UE::LLMPrivate::FLLMGlobals::GetInner();
-		// We have to check IsEnabled() again after calling Get, because the constructor is called
-		// from Get, and will set EnabledState=Disabled if the platform doesn't support it.
-		if (FLowLevelMemTracker::IsEnabled())
+	Init(InTagSet, InTracker, bOverride,
+		// ShouldEnable
+		[](FLLMGlobals& LLMRef, FLLMTracker* TrackerPtr)
 		{
-			LLMRef.BootstrapInitialise();
-
-			if (bOverride || LLMRef.GetTracker(InTracker)->GetActiveTagData(InTagSet) == nullptr)
-			{
-				bEnabled = true;
-				Tracker = InTracker;
-				TagSet = InTagSet;
-				LLMRef.GetTracker(Tracker)->PushTag(TagEnum, InTagSet);
-			}
-		}
-	};
+			return true;
+		},
+		// OnEnable
+		[TagEnum, bInIsStatTag, InTagSet, InTracker](FLLMGlobals&, FLLMTracker* TrackerPtr)
+		{
+			// This scope does not support TagSets (use FLLMScopeDynamic instead), except for the special case of
+			// anonymized allocations, which are indicated with a TagSet and ELLMTag::EngineMisc
+			LLMCheck((!bInIsStatTag && InTagSet == ELLMTagSet::None) || (TagEnum == ELLMTag::EngineMisc));
+			// ELLMTag::FMalloc is a special tag expected to only be used by the Platform Tracker (see header where defined).
+			LLMCheck((TagEnum != ELLMTag::FMalloc) || (ELLMTracker::Platform == InTracker));
 
-	// But remember that if we abort while we hold the scope, we need to destroy the scope.
-	AutoRTFM::PushOnAbortHandler(this, [this] { if (bEnabled) { Destruct(); } });
+			TrackerPtr->PushTag(TagEnum, InTagSet);
+		});
 }
 
 FLLMScope::FLLMScope(FName TagName, bool bInIsStatTag, ELLMTagSet InTagSet, ELLMTracker InTracker, bool bOverride)
 {
-	if (!FLowLevelMemTracker::IsEnabled())
-	{
-		return;
-	}
+	using namespace UE::LLMPrivate;
 
-	// We run the init in the open, because we want to track LLM even in transactions.
-	UE_AUTORTFM_OPEN
-	{
-		UE::LLMPrivate::FLLMGlobals& LLMRef = UE::LLMPrivate::FLLMGlobals::GetInner();
-		// We have to check IsEnabled() again after calling Get, because the constructor is called
-		// from Get, and will set EnabledStae=Disabled if the platform doesn't support it.
-		if (FLowLevelMemTracker::IsEnabled())
+	Init(InTagSet, InTracker, bOverride,
+		// ShouldEnable
+		[InTagSet](FLLMGlobals& LLMRef, FLLMTracker* TrackerPtr)
 		{
-			LLMRef.BootstrapInitialise();
-			if (LLMRef.IsTagSetScopeActiveInner(InTagSet))
+			return LLMRef.IsTagSetScopeActiveInner(InTagSet);
+		},
+		// OnEnable
+		[TagName, bInIsStatTag, InTagSet](FLLMGlobals& LLMRef, FLLMTracker* TrackerPtr)
+		{
+			TrackerPtr->PushTag(TagName, bInIsStatTag, InTagSet);
+		});
+}
+
+FLLMScope::FLLMScope(const UE::LLMPrivate::FTagData* TagData, bool bInIsStatTag, ELLMTagSet InTagSet,
+	ELLMTracker InTracker, bool bOverride)
+{
+	using namespace UE::LLMPrivate;
+
+	Init(InTagSet, InTracker, bOverride,
+		// ShouldEnable
+		[InTagSet](FLLMGlobals& LLMRef, FLLMTracker* TrackerPtr)
+		{
+			return true;
+		},
+		// OnEnable
+		[TagData, bInIsStatTag, InTagSet, this](FLLMGlobals& LLMRef, FLLMTracker* TrackerPtr)
+		{
+			TrackerPtr->PushTag(TagData, InTagSet);
+			if (InTagSet == ELLMTagSet::AssetClasses && LLMRef.IsTagSetRecordingActiveInner(ELLMTagSet::CodeOrContent))
 			{
-				if (bOverride || LLMRef.GetTracker(InTracker)->GetActiveTagData(InTagSet) == nullptr)
+				// AssetClasses scopes also have the responsibility to trigger CodeOrContent scopes, so when
+				// pushing onto AssetClasses, also push onto CodeOrContent.
+				bPushedCodeOrContent = true;
+				if (!TagData)
 				{
-					bEnabled = true;
-					Tracker = InTracker;
-					TagSet = InTagSet;
-
-					LLMRef.GetTracker(Tracker)->PushTag(TagName, bInIsStatTag, InTagSet);
+					TrackerPtr->PushTag(nullptr, ELLMTagSet::CodeOrContent);
+				}
+				else
+				{
+					const FTagData* CodeOrContentData = LLMRef.FindOrAddCodeOrContentContentTagData();
+					TrackerPtr->PushTag(CodeOrContentData, ELLMTagSet::CodeOrContent);
 				}
 			}
-		}
-	};
-
-	// But remember that if we abort while we hold the scope, we need to destroy the scope.
-	AutoRTFM::PushOnAbortHandler(this, [this] { if (bEnabled) { Destruct(); } });
+		});
 }
 
-FLLMScope::FLLMScope(const UE::LLMPrivate::FTagData* TagData, bool bInIsStatTag, ELLMTagSet InTagSet,
-	ELLMTracker InTracker, bool bOverride)
+template <typename ShouldEnableType, typename OnEnabledType>
+void FLLMScope::Init(ELLMTagSet InTagSet, ELLMTracker InTracker, bool bOverride,
+	ShouldEnableType&& ShouldEnable, OnEnabledType&& OnEnable)
 {
 	using namespace UE::LLMPrivate;
 
@@ -3779,50 +3782,71 @@ FLLMScope::FLLMScope(const UE::LLMPrivate::FTagData* TagData, bool bInIsStatTag,
 		return;
 	}
 
-	// We run the init in the open, because we want to track LLM even in transactions.
+	// Note that bEnabled=false because the constructor initializes bEnabled=false before calling the Init function.
+	// That original initialization of bEnabled=false occurs in the open, but it does not get rolled back because our
+	// data is on the stack and the stack is not rolled back. We rely on our data not being rolled back; we have
+	// pointers to it from the lambda we pass to PushAbortHandler.
+	LLMCheck(!bEnabled);
+	FLLMGlobals& LLMRef = FLLMGlobals::GetInner();
+	// Scopes can be encountered in PreMain during AutoRTFM startup, before IsOnCurrentTransactionStack is
+	// available. Do not call it until after LLM is fully initialized.
+	LLMCheck(!LLMRef.bFullyInitialized || !AutoRTFM::IsClosed() || AutoRTFM::IsOnCurrentTransactionStack(this));
+
+	// We run Init in the open, because we want to track LLM even in transactions.
 	UE_AUTORTFM_OPEN
 	{
-		UE::LLMPrivate::FLLMGlobals& LLMRef = UE::LLMPrivate::FLLMGlobals::GetInner();
 		// We have to check IsEnabled() again after calling Get, because the constructor is called
 		// from Get, and will set EnabledState=Disabled if the platform doesn't support it.
 		if (FLowLevelMemTracker::IsEnabled())
 		{
 			LLMRef.BootstrapInitialise();
 
-			if (bOverride || LLMRef.GetTracker(InTracker)->GetActiveTagData(InTagSet) == nullptr)
+			FLLMTracker* TrackerPtr = LLMRef.GetTracker(InTracker);
+			if (bOverride || TrackerPtr->GetActiveTagData(InTagSet) == nullptr)
 			{
-				bEnabled = true;
-				Tracker = InTracker;
-				TagSet = InTagSet;
-				LLMRef.GetTracker(Tracker)->PushTag(TagData, InTagSet);
-				if (TagSet == ELLMTagSet::AssetClasses && LLMRef.IsTagSetRecordingActiveInner(ELLMTagSet::CodeOrContent))
+				if (ShouldEnable(LLMRef, TrackerPtr))
 				{
-					// AssetClasses scopes also have the responsibility to trigger CodeOrContent scopes, so when
-					// pushing onto AssetClasses,onto CodeOrContent.
-					bPushedCodeOrContent = true;
-					if (!TagData)
-					{
-						LLMRef.GetTracker(Tracker)->PushTag(nullptr, ELLMTagSet::CodeOrContent);
-					}
-					else
-					{
-						const FTagData* CodeOrContentData = LLMRef.FindOrAddCodeOrContentContentTagData();
-						LLMRef.GetTracker(Tracker)->PushTag(CodeOrContentData, ELLMTagSet::CodeOrContent);
-					}
+					bEnabled = true;
+					Tracker = InTracker;
+					TagSet = InTagSet;
+					OnEnable(LLMRef, TrackerPtr);
 				}
 			}
 		}
 	};
 
 	// But remember that if we abort while we hold the scope, we need to destroy the scope.
-	AutoRTFM::PushOnAbortHandler(this, [this] { if (bEnabled) { Destruct(); } });
+	if (bEnabled)
+	{
+		AutoRTFM::PushOnAbortHandler(this, [this]
+			{
+				if (bEnabled)
+				{
+					DestructInTheOpen();
+					bEnabled = false;
+				}
+			});
+	}
 }
 
 void FLLMScope::Destruct()
+{
+	check(bEnabled);
+	// We run Destruct in the open, because we want to track LLM even in transactions.
+	UE_AUTORTFM_OPEN
+	{
+		DestructInTheOpen();
+		bEnabled = false;
+	};
+	// But remember to pop our on-abort handler (so that we don't try and double Destruct).
+	AutoRTFM::PopOnAbortHandler(this);
+}
+
+void FLLMScope::DestructInTheOpen()
 {
 	using namespace UE::LLMPrivate;
 
-	FLLMGlobals& LLMRef = UE::LLMPrivate::FLLMGlobals::GetInner();
+	FLLMGlobals& LLMRef = FLLMGlobals::GetInner();
 	FLLMTracker* TrackerPtr = LLMRef.GetTracker(Tracker);
 	TrackerPtr->PopTag(TagSet);
 	if (bPushedCodeOrContent)
@@ -3831,129 +3855,180 @@ void FLLMScope::Destruct()
 	}
 }
 
-FLLMScope::~FLLMScope()
-{
-	if (bEnabled)
-	{
-		// We run the destroy in the open, because we want to track LLM even in transactions.
-		UE_AUTORTFM_OPEN{ Destruct(); };
-		// But remember to pop our on-abort handler (so that we don't try and double destroy).
-		AutoRTFM::PopOnAbortHandler(this);
-	}
-}
-
 void FLLMScopeDynamic::Init(ELLMTracker InTracker, ELLMTagSet InTagSet)
 {
+	// Note that bEnabled=false because the constructor initializes bEnabled=false before calling the Init function.
+	// That original initialization of bEnabled=false occurs in the open, but it does not get rolled back because our
+	// data is on the stack and the stack is not rolled back. We rely on our data not being rolled back; we have
+	// pointers to it from the lambda we pass to PushAbortHandler.
+	LLMCheck(!bEnabled);
 	UE::LLMPrivate::FLLMGlobals& LLMRef = UE::LLMPrivate::FLLMGlobals::GetInner();
-	// We have to check IsEnabled() again after calling Get, because the constructor is called
-	// from Get, and will set EnabledState=Disabled if the platform doesn't support it.
-	if (!FLowLevelMemTracker::IsEnabled())
+	// Scopes can be encountered in PreMain during AutoRTFM startup, before IsOnCurrentTransactionStack is
+	// available. Do not call it until after LLM is fully initialized.
+	LLMCheck(!LLMRef.bFullyInitialized || !AutoRTFM::IsClosed() || AutoRTFM::IsOnCurrentTransactionStack(this));
+
+	// We run the init in the open, because we want to track LLM even in transactions.
+	UE_AUTORTFM_OPEN
 	{
-		bEnabled = false;
-		return;
-	}
-	LLMRef.BootstrapInitialise();
-	if (!LLMRef.IsTagSetScopeActiveInner(InTagSet))
+		// We have to check IsEnabled() again after calling Get, because the constructor is called
+		// from Get, and will set EnabledState=Disabled if the platform doesn't support it.
+		if (!FLowLevelMemTracker::IsEnabled())
+		{
+			return; // Note that depending on config this might exit out of UE_AUTORTFM_OPEN block only
+		}
+
+		LLMRef.BootstrapInitialise();
+		if (!LLMRef.IsTagSetScopeActiveInner(InTagSet))
+		{
+			return; // Note that depending on config this might exit out of UE_AUTORTFM_OPEN block only
+		}
+		bEnabled = true;
+		TagData = nullptr;
+		CodeOrContentTagData = nullptr;
+		Tracker = InTracker;
+		TagSet = InTagSet;
+	};
+
+	// But remember that if we abort while we hold the scope, we need to destroy the scope.
+	if (bEnabled)
 	{
-		bEnabled = false;
-		return;
+		AutoRTFM::PushOnAbortHandler(this, [this]
+			{
+				if (bEnabled)
+				{
+					DestructInTheOpen();
+					bEnabled = false;
+				}
+			});
 	}
-	bEnabled = true;
-	TagData = nullptr;
-	CodeOrContentTagData = nullptr;
-	Tracker = InTracker;
-	TagSet = InTagSet;
 }
 
 bool FLLMScopeDynamic::TryFindTag(FName TagName)
 {
-	// AssetClasses scopes also have the responsibility to trigger CodeOrContent scopes, so we might
-	// reach this point even if TagSet is not RecordingActive. In the case where the TagSet is not
-	// recording, we should not create its tag and should return true from TryFindTag.
+	// Like Init, this function can read/write LLM's internal data (and uses non-transactional-safe locks) and
+	// therefore we want to run it in the open. There is nothing we need to revert after an abort other than what has
+	// already been set up by Init, so no need for any additional AbortHandler.
 	bool bNeedAddTag = false;
-	UE::LLMPrivate::FLLMGlobals& LLMRef = UE::LLMPrivate::FLLMGlobals::GetInner();
-	if (LLMRef.IsTagSetRecordingActiveInner(TagSet))
-	{
-		TagData = LLMRef.FindTagData(TagName, TagSet, UE::LLMPrivate::ETagReferenceSource::Scope);
-		bNeedAddTag = TagData == nullptr;
-	}
-	if (TagSet == ELLMTagSet::AssetClasses && LLMRef.IsTagSetRecordingActiveInner(ELLMTagSet::CodeOrContent))
-	{
-		CodeOrContentTagData = LLMRef.FindOrAddCodeOrContentContentTagData();
-	}
-	else
+	UE_AUTORTFM_OPEN
 	{
-		CodeOrContentTagData = nullptr;
-	}
+		// AssetClasses scopes also have the responsibility to trigger CodeOrContent scopes, so we might
+		// reach this point even if TagSet is not RecordingActive. In the case where the TagSet is not
+		// recording, we should not create its tag and should return true from TryFindTag.
+		UE::LLMPrivate::FLLMGlobals& LLMRef = UE::LLMPrivate::FLLMGlobals::GetInner();
+
+		// Setting this->TagData and this->CodeOrContentTagData notifies DestructInTheOpen to unpop them,
+		// whether called from RTFM AbortHandler or from this class's destructor. We rely on Activate or
+		// TryAddTagAndActivate being called to do the corresponding Push after we set these variables.
+		if (LLMRef.IsTagSetRecordingActiveInner(TagSet))
+		{
+			TagData = LLMRef.FindTagData(TagName, TagSet, UE::LLMPrivate::ETagReferenceSource::Scope);
+			bNeedAddTag = TagData == nullptr;
+		}
+		if (TagSet == ELLMTagSet::AssetClasses && LLMRef.IsTagSetRecordingActiveInner(ELLMTagSet::CodeOrContent))
+		{
+			CodeOrContentTagData = LLMRef.FindOrAddCodeOrContentContentTagData();
+		}
+		else
+		{
+			CodeOrContentTagData = nullptr;
+		}
+	};
 	return !bNeedAddTag;
 }
 
 bool FLLMScopeDynamic::TryAddTagAndActivate(FName TagName, const ILLMDynamicTagConstructor& Constructor)
 {
-	UE::LLMPrivate::FLLMGlobals& LLMRef = UE::LLMPrivate::FLLMGlobals::GetInner();
-	FName StatFullName = NAME_None;
-#if STATS
-	FString StatConstructorName = Constructor.GetStatName();
-	if (!StatConstructorName.IsEmpty())
+	// Like Init, this function can read/write LLM's internal data (and uses non-transactional-safe locks) and
+	// therefore we want to run it in the open. There is nothing we need to revert after an abort other than what has
+	// already been set up by Init, so no need for any additional AbortHandler.
+	UE_AUTORTFM_OPEN
 	{
-		if (Constructor.NeedsStatConstruction())
+		UE::LLMPrivate::FLLMGlobals & LLMRef = UE::LLMPrivate::FLLMGlobals::GetInner();
+		FName StatFullName = NAME_None;
+#if STATS
+		FString StatConstructorName = Constructor.GetStatName();
+		if (!StatConstructorName.IsEmpty())
 		{
-			switch (TagSet)
+			if (Constructor.NeedsStatConstruction())
 			{
-			case ELLMTagSet::None:
-				StatFullName = FDynamicStats::CreateMemoryStatId<FStatGroup_STATGROUP_LLMFULL>(
-					StatConstructorName).GetName();
-				break;
-			case ELLMTagSet::Assets:
-				StatFullName = FDynamicStats::CreateMemoryStatId<FStatGroup_STATGROUP_LLMAssets>(
-					StatConstructorName).GetName();
-				break;
-			case ELLMTagSet::AssetClasses:
-				StatFullName = FDynamicStats::CreateMemoryStatId<FStatGroup_STATGROUP_LLMAssets>(
-					StatConstructorName).GetName();
-				break;
-			case ELLMTagSet::UObjectClasses:
-				StatFullName = FDynamicStats::CreateMemoryStatId<FStatGroup_STATGROUP_LLMUObjects>(
-					StatConstructorName).GetName();
-				break;
-			case ELLMTagSet::CodeOrContent:
-				StatFullName = FDynamicStats::CreateMemoryStatId<FStatGroup_STATGROUP_LLMCodeOrContent>(
-					StatConstructorName).GetName();
-				break;
-			default:
-				checkNoEntry();
-				break;
+				switch (TagSet)
+				{
+				case ELLMTagSet::None:
+					StatFullName = FDynamicStats::CreateMemoryStatId<FStatGroup_STATGROUP_LLMFULL>(
+						StatConstructorName).GetName();
+					break;
+				case ELLMTagSet::Assets:
+					StatFullName = FDynamicStats::CreateMemoryStatId<FStatGroup_STATGROUP_LLMAssets>(
+						StatConstructorName).GetName();
+					break;
+				case ELLMTagSet::AssetClasses:
+					StatFullName = FDynamicStats::CreateMemoryStatId<FStatGroup_STATGROUP_LLMAssets>(
+						StatConstructorName).GetName();
+					break;
+				case ELLMTagSet::UObjectClasses:
+					StatFullName = FDynamicStats::CreateMemoryStatId<FStatGroup_STATGROUP_LLMUObjects>(
+						StatConstructorName).GetName();
+					break;
+				case ELLMTagSet::CodeOrContent:
+					StatFullName = FDynamicStats::CreateMemoryStatId<FStatGroup_STATGROUP_LLMCodeOrContent>(
+						StatConstructorName).GetName();
+					break;
+				default:
+					checkNoEntry();
+					break;
+				}
+			}
+			else
+			{
+				StatFullName = FName(*StatConstructorName);
 			}
 		}
-		else
+#endif
+		// Setting this->TagData and this->CodeOrContentTagData notifies DestructInTheOpen to unpop them,
+		// whether called from RTFM AbortHandler or from this class's destructor.
+		TagData = LLMRef.FindOrAddTagData(TagName, TagSet, StatFullName, UE::LLMPrivate::ETagReferenceSource::Scope);
+		LLMRef.GetTracker(Tracker)->PushTag(TagData, TagSet);
+		if (CodeOrContentTagData)
 		{
-			StatFullName = FName(*StatConstructorName);
+			LLMRef.GetTracker(Tracker)->PushTag(CodeOrContentTagData, ELLMTagSet::CodeOrContent);
 		}
-	}
-#endif
-	TagData = LLMRef.FindOrAddTagData(TagName, TagSet, StatFullName, UE::LLMPrivate::ETagReferenceSource::Scope);
-	LLMRef.GetTracker(Tracker)->PushTag(TagData, TagSet);
-	if (CodeOrContentTagData)
-	{
-		LLMRef.GetTracker(Tracker)->PushTag(CodeOrContentTagData, ELLMTagSet::CodeOrContent);
-	}
+	};
 	return true;
 }
 
 void FLLMScopeDynamic::Activate()
 {
-	UE::LLMPrivate::FLLMGlobals& LLMRef = UE::LLMPrivate::FLLMGlobals::GetInner();
-	if (TagData)
-	{
-		LLMRef.GetTracker(Tracker)->PushTag(TagData, TagSet);
-	}
-	if (CodeOrContentTagData)
+	// Like Init, this function can read/write LLM's internal data (and uses non-transactional-safe locks) and
+	// therefore we want to run it in the open. There is nothing we need to revert after an abort other than what has
+	// already been set up by Init, so no need for any additional AbortHandler.
+	UE_AUTORTFM_OPEN
 	{
-		LLMRef.GetTracker(Tracker)->PushTag(CodeOrContentTagData, ELLMTagSet::CodeOrContent);
-	}
+		UE::LLMPrivate::FLLMGlobals & LLMRef = UE::LLMPrivate::FLLMGlobals::GetInner();
+		if (TagData)
+		{
+			LLMRef.GetTracker(Tracker)->PushTag(TagData, TagSet);
+		}
+		if (CodeOrContentTagData)
+		{
+			LLMRef.GetTracker(Tracker)->PushTag(CodeOrContentTagData, ELLMTagSet::CodeOrContent);
+		}
+	};
 }
 
 void FLLMScopeDynamic::Destruct()
+{
+	check(bEnabled);
+	// We run Destruct in the open, because we want to track LLM even in transactions.
+	UE_AUTORTFM_OPEN
+	{
+		DestructInTheOpen();
+		bEnabled = false;
+	};
+	// But remember to pop our on-abort handler (so that we don't try and double Destruct).
+	AutoRTFM::PopOnAbortHandler(this);
+}
+
+void FLLMScopeDynamic::DestructInTheOpen()
 {
 	UE::LLMPrivate::FLLMGlobals& LLMRef = UE::LLMPrivate::FLLMGlobals::GetInner();
 	if (TagData)
@@ -3971,7 +4046,6 @@ FLLMPauseScope::FLLMPauseScope(FName TagName, bool bIsStatTag, uint64 Amount, EL
 {
 	if (!FLowLevelMemTracker::IsEnabled())
 	{
-		bEnabled = false;
 		return;
 	}
 	Init(TagName, ELLMTag::Untagged, false, bIsStatTag, Amount, TrackerToPause, InAllocType);
@@ -3982,7 +4056,6 @@ FLLMPauseScope::FLLMPauseScope(ELLMTag TagEnum, bool bIsStatTag, uint64 Amount,
 {
 	if (!FLowLevelMemTracker::IsEnabled())
 	{
-		bEnabled = false;
 		return;
 	}
 	LLMCheck(!bIsStatTag);
@@ -3992,57 +4065,91 @@ FLLMPauseScope::FLLMPauseScope(ELLMTag TagEnum, bool bIsStatTag, uint64 Amount,
 void FLLMPauseScope::Init(FName TagName, ELLMTag EnumTag, bool bIsEnumTag, bool bIsStatTag, uint64 Amount,
 	ELLMTracker TrackerToPause, ELLMAllocType InAllocType)
 {
+	// Note that bEnabled=false because the constructor initializes bEnabled=false before calling the Init function.
+	// That original initialization of bEnabled=false occurs in the open, but it does not get rolled back because our
+	// data is on the stack and the stack is not rolled back. We rely on our data not being rolled back; we have
+	// pointers to it from the lambda we pass to PushAbortHandler.
+	LLMCheck(!bEnabled);
 	UE::LLMPrivate::FLLMGlobals& LLMRef = UE::LLMPrivate::FLLMGlobals::GetInner();
-	// We have to check IsEnabled() again after calling Get, because the constructor is called
-	// from Get, and will set EnabledState=Disabled if the platform doesn't support it.
-	if (!FLowLevelMemTracker::IsEnabled())
-	{
-		bEnabled = false;
-		return;
-	}
-	LLMRef.BootstrapInitialise();
-	if (!LLMRef.IsTagSetScopeActiveInner(ELLMTagSet::None))
+	// Scopes can be encountered in PreMain during AutoRTFM startup, before IsOnCurrentTransactionStack is
+	// available. Do not call it until after LLM is fully initialized.
+	LLMCheck(!LLMRef.bFullyInitialized || !AutoRTFM::IsClosed() || AutoRTFM::IsOnCurrentTransactionStack(this));
+
+	// We run the init in the open, because we want to track LLM even in transactions.
+	UE_AUTORTFM_OPEN
 	{
-		bEnabled = false;
-		return;
-	}
+		// We have to check IsEnabled() again after calling Get, because the constructor is called
+		// from Get, and will set EnabledState=Disabled if the platform doesn't support it.
+		if (!FLowLevelMemTracker::IsEnabled())
+		{
+			return; // Note that depending on config this might exit out of UE_AUTORTFM_OPEN block only
+		}
 
-	bEnabled = true;
-	PausedTracker = TrackerToPause;
-	AllocType = InAllocType;
+		LLMRef.BootstrapInitialise();
+		if (!LLMRef.IsTagSetScopeActiveInner(ELLMTagSet::None))
+		{
+			return; // Note that depending on config this might exit out of UE_AUTORTFM_OPEN block only
+		}
 
-	for (int32 TrackerIndex = 0; TrackerIndex < static_cast<int32>(ELLMTracker::Max); TrackerIndex++)
-	{
-		ELLMTracker Tracker = (ELLMTracker)TrackerIndex;
+		bEnabled = true;
+		PausedTracker = TrackerToPause;
+		AllocType = InAllocType;
 
-		if (PausedTracker == ELLMTracker::Max || PausedTracker == Tracker)
+		for (int32 TrackerIndex = 0; TrackerIndex < static_cast<int32>(ELLMTracker::Max); TrackerIndex++)
 		{
-			if (Amount == 0)
-			{
-				LLMRef.GetTracker(Tracker)->Pause(InAllocType);
-			}
-			else
+			ELLMTracker Tracker = (ELLMTracker)TrackerIndex;
+
+			if (PausedTracker == ELLMTracker::Max || PausedTracker == Tracker)
 			{
-				if (bIsEnumTag)
+				if (Amount == 0)
 				{
-					LLMRef.GetTracker(Tracker)->PauseAndTrackMemory(EnumTag, static_cast<int64>(Amount), InAllocType);
+					LLMRef.GetTracker(Tracker)->Pause(InAllocType);
 				}
 				else
 				{
-					LLMRef.GetTracker(Tracker)->PauseAndTrackMemory(TagName, ELLMTagSet::None, bIsStatTag,
-						static_cast<int64>(Amount), InAllocType);
+					if (bIsEnumTag)
+					{
+						LLMRef.GetTracker(Tracker)->PauseAndTrackMemory(EnumTag, static_cast<int64>(Amount), InAllocType);
+					}
+					else
+					{
+						LLMRef.GetTracker(Tracker)->PauseAndTrackMemory(TagName, ELLMTagSet::None, bIsStatTag,
+							static_cast<int64>(Amount), InAllocType);
+					}
 				}
 			}
 		}
+	};
+
+	// But remember that if we abort while we hold the scope, we need to destroy the scope.
+	if (bEnabled)
+	{
+		AutoRTFM::PushOnAbortHandler(this, [this]
+			{
+				if (bEnabled)
+				{
+					DestructInTheOpen();
+					bEnabled = false;
+				}
+			});
 	}
 }
 
-FLLMPauseScope::~FLLMPauseScope()
+void FLLMPauseScope::Destruct()
 {
-	if (!bEnabled)
+	check(bEnabled);
+	// We run Destruct in the open, because we want to track LLM even in transactions.
+	UE_AUTORTFM_OPEN
 	{
-		return;
-	}
+		DestructInTheOpen();
+		bEnabled = false;
+	};
+	// But remember to pop our on-abort handler (so that we don't try and double Destruct).
+	AutoRTFM::PopOnAbortHandler(this);
+}
+
+void FLLMPauseScope::DestructInTheOpen()
+{
 	UE::LLMPrivate::FLLMGlobals& LLMRef = UE::LLMPrivate::FLLMGlobals::GetInner();
 
 	for (int32 TrackerIndex = 0; TrackerIndex < static_cast<int32>(ELLMTracker::Max); TrackerIndex++)
@@ -4064,65 +4171,102 @@ FLLMClearScope::FLLMClearScope(ELLMTagSet InTagSet, ELLMTracker InTracker)
 FLLMScopeFromPtr::FLLMScopeFromPtr(void* Ptr, ELLMTracker InTracker)
 {
 	using namespace UE::LLMPrivate;
+
 	if (!FLowLevelMemTracker::IsEnabled())
 	{
-		DisableAll();
 		return;
 	}
 	if (Ptr == nullptr)
 	{
-		DisableAll();
 		return;
 	}
 
+	// Note that bEnabled=false because the constructor initializes bEnabled=false before calling the Init function.
+	// That original initialization of bEnabled=false occurs in the open, but it does not get rolled back because our
+	// data is on the stack and the stack is not rolled back. We rely on our data not being rolled back; we have
+	// pointers to it from the lambda we pass to PushAbortHandler.
+	LLMCheck(!bEnabled);
 	FLLMGlobals& LLMRef = FLLMGlobals::GetInner();
-	// We have to check IsEnabled() again after calling Get, because the constructor is called
-	// from Get, and will set EnabledState=Disabled if the platform doesn't support it.
-	if (!FLowLevelMemTracker::IsEnabled())
-	{
-		DisableAll();
-		return;
-	}
-	LLMRef.BootstrapInitialise();
+	// Scopes can be encountered in PreMain during AutoRTFM startup, before IsOnCurrentTransactionStack is
+	// available. Do not  call it until after LLM is fully initialized.
+	LLMCheck(!LLMRef.bFullyInitialized || !AutoRTFM::IsClosed() || AutoRTFM::IsOnCurrentTransactionStack(this));
 
-	FLLMTracker* TrackerData = LLMRef.GetTracker(InTracker);
-
-	TArray<const FTagData*, TInlineAllocator<static_cast<int32>(ELLMTagSet::Max)>> Tags;
-	if (!TrackerData->FindTagsForPtr(Ptr, Tags))
-	{
-		DisableAll();
-		return;
-	}
-
-	Tracker = InTracker;
-
-	for (int32 TagSetAsInteger = 0; TagSetAsInteger < static_cast<int32>(ELLMTagSet::Max); TagSetAsInteger++)
+	// We run the constructor in the open, because we want to track LLM even in transactions.
+	UE_AUTORTFM_OPEN
 	{
-		if (!LLMRef.IsTagSetScopeActiveInner(static_cast<ELLMTagSet>(TagSetAsInteger)))
+		// We have to check IsEnabled() again after calling Get, because the constructor is called
+		// from Get, and will set EnabledState=Disabled if the platform doesn't support it.
+		if (!FLowLevelMemTracker::IsEnabled())
 		{
-			bEnabled[TagSetAsInteger] = false;
-
-			continue;
+			return; // Note that depending on config this might exit out of UE_AUTORTFM_OPEN only
 		}
 
-		if (Tags[TagSetAsInteger])
+		LLMRef.BootstrapInitialise();
+
+		FLLMTracker* TrackerData = LLMRef.GetTracker(InTracker);
+
+		TArray<const FTagData*, TInlineAllocator<static_cast<int32>(ELLMTagSet::Max)>> Tags;
+		if (!TrackerData->FindTagsForPtr(Ptr, Tags))
 		{
-			bEnabled[TagSetAsInteger] = true;
-			TrackerData->PushTag(Tags[TagSetAsInteger], static_cast<ELLMTagSet>(TagSetAsInteger));
+			return; // Note that depending on config this might exit out of UE_AUTORTFM_OPEN only
 		}
-		else
+
+		bEnabled = true;
+		Tracker = InTracker;
+
+		for (int32 TagSetAsInteger = 0; TagSetAsInteger < static_cast<int32>(ELLMTagSet::Max); TagSetAsInteger++)
 		{
-			bEnabled[TagSetAsInteger] = false;
+			if (!LLMRef.IsTagSetScopeActiveInner(static_cast<ELLMTagSet>(TagSetAsInteger)))
+			{
+				bSetEnabled[TagSetAsInteger] = false;
+				continue;
+			}
+
+			if (Tags[TagSetAsInteger])
+			{
+				bSetEnabled[TagSetAsInteger] = true;
+				TrackerData->PushTag(Tags[TagSetAsInteger], static_cast<ELLMTagSet>(TagSetAsInteger));
+			}
+			else
+			{
+				bSetEnabled[TagSetAsInteger] = false;
+			}
 		}
+	};
+
+	// But remember that if we abort while we hold the scope, we need to destroy the scope.
+	if (bEnabled)
+	{
+		AutoRTFM::PushOnAbortHandler(this, [this]
+			{
+				if (bEnabled)
+				{
+					DestructInTheOpen();
+					bEnabled = false;
+				}
+			});
 	}
 }
 
-FLLMScopeFromPtr::~FLLMScopeFromPtr()
+void FLLMScopeFromPtr::Destruct()
+{
+	check(bEnabled);
+	// We run Destruct in the open, because we want to track LLM even in transactions.
+	UE_AUTORTFM_OPEN
+	{
+		DestructInTheOpen();
+		bEnabled = false;
+	};
+	// But remember to pop our on-abort handler (so that we don't try and double Destruct).
+	AutoRTFM::PopOnAbortHandler(this);
+}
+
+void FLLMScopeFromPtr::DestructInTheOpen()
 {
 	UE::LLMPrivate::FLLMGlobals& LLMRef = UE::LLMPrivate::FLLMGlobals::GetInner();
 	for (int32 TagSetAsInteger = 0; TagSetAsInteger < static_cast<int32>(ELLMTagSet::Max); TagSetAsInteger++)
 	{
-		if (bEnabled[TagSetAsInteger])
+		if (bSetEnabled[TagSetAsInteger])
 		{
 			LLMRef.GetTracker(Tracker)->PopTag(static_cast<ELLMTagSet>(TagSetAsInteger));
 		}
diff --git a/Engine/Source/Runtime/Core/Public/HAL/LowLevelMemTracker.h b/Engine/Source/Runtime/Core/Public/HAL/LowLevelMemTracker.h
--- a/Engine/Source/Runtime/Core/Public/HAL/LowLevelMemTracker.h
+++ b/Engine/Source/Runtime/Core/Public/HAL/LowLevelMemTracker.h
@@ -745,10 +745,20 @@ public:
 	CORE_API explicit FLLMScope(FName TagName, bool bIsStatTag, ELLMTagSet InTagSet, ELLMTracker InTracker, bool bOverride = true);
 	CORE_API explicit FLLMScope(ELLMTag TagEnum, bool bIsStatTag, ELLMTagSet InTagSet, ELLMTracker InTracker, bool bOverride = true);
 	CORE_API explicit FLLMScope(const UE::LLMPrivate::FTagData* TagData, bool bIsStatTag, ELLMTagSet InTagSet, ELLMTracker InTracker, bool bOverride = true);
-	CORE_API ~FLLMScope();
+	~FLLMScope()
+	{
+		if (bEnabled)
+		{
+			Destruct();
+		}
+	}
 
 protected:
 	CORE_API void Destruct();
+	AUTORTFM_DISABLE void DestructInTheOpen();
+	template <typename ShouldEnableType, typename OnEnabledType>
+	void Init(ELLMTagSet InTagSet, ELLMTracker InTracker, bool bOverride,
+		ShouldEnableType&& ShouldEnable, OnEnabledType&& OnEnable);
 
 	// All fields other than bEnabled are not initialized by constructor unless bEnabled is true
 	ELLMTracker Tracker;
@@ -816,6 +826,7 @@ public:
 protected:
 	CORE_API void Init(ELLMTracker InTracker, ELLMTagSet InTagSet);
 	CORE_API void Destruct();
+	AUTORTFM_DISABLE void DestructInTheOpen();
 
 	// All fields other than bEnabled are not initialized by constructor unless bEnabled is true
 	const UE::LLMPrivate::FTagData* TagData;
@@ -831,13 +842,23 @@ class FLLMPauseScope
 public:
 	CORE_API FLLMPauseScope(FName TagName, bool bIsStatTag, uint64 Amount, ELLMTracker TrackerToPause, ELLMAllocType InAllocType);
 	CORE_API FLLMPauseScope(ELLMTag TagEnum, bool bIsStatTag, uint64 Amount, ELLMTracker TrackerToPause, ELLMAllocType InAllocType);
-	CORE_API ~FLLMPauseScope();
+	~FLLMPauseScope()
+	{
+		if (bEnabled)
+		{
+			Destruct();
+		}
+	}
 protected:
 	CORE_API void Init(FName TagName, ELLMTag EnumTag, bool bIsEnumTag, bool bIsStatTag, uint64 Amount,
 		ELLMTracker TrackerToPause, ELLMAllocType InAllocType);
+	CORE_API void Destruct();
+	AUTORTFM_DISABLE void DestructInTheOpen();
+
+	// All fields other than bEnabled are not initialized by constructor unless bEnabled is true
 	ELLMTracker PausedTracker;
 	ELLMAllocType AllocType;
-	bool bEnabled;
+	bool bEnabled = false;
 };
 
 /** LLM Scope to clear top tag while in scope */
@@ -852,16 +873,22 @@ class FLLMScopeFromPtr
 {
 public:
 	CORE_API FLLMScopeFromPtr(void* Ptr, ELLMTracker Tracker);
-	CORE_API ~FLLMScopeFromPtr();
-protected:
-	ELLMTracker Tracker;
-	bool bEnabled[static_cast<int32>(ELLMTagSet::Max)];
-
-private:
-	inline void DisableAll()
+	~FLLMScopeFromPtr()
 	{
-		FMemory::Memzero(bEnabled);
+		if (bEnabled)
+		{
+			Destruct();
+		}
 	}
+
+protected:
+	CORE_API void Destruct();
+	AUTORTFM_DISABLE void DestructInTheOpen();
+
+	// All fields other than bEnabled are not initialized by constructor unless bEnabled is true
+	ELLMTracker Tracker;
+	bool bSetEnabled[static_cast<int32>(ELLMTagSet::Max)];
+	bool bEnabled = false;
 };
 
 /** Global instances to provide information about a tag to LLM. */
