Add log categories and relative timestamp since start to early crash logs.

This affects early crashes that happen before the file log is fully ready, and while it's holding all logging into an in-memory backlog. The crash handler retrieves this backlog and writes it into the crash log, but by default does not include any log category or timestamp. This new option allows to use the time information stored in the backlog messages to write out the timestamp since start for each log line.

diff --git a/Engine/Source/Runtime/Core/Private/GenericPlatform/GenericPlatformCrashContext.cpp b/Engine/Source/Runtime/Core/Private/GenericPlatform/GenericPlatformCrashContext.cpp
--- a/Engine/Source/Runtime/Core/Private/GenericPlatform/GenericPlatformCrashContext.cpp
+++ b/Engine/Source/Runtime/Core/Private/GenericPlatform/GenericPlatformCrashContext.cpp
@@ -26,6 +26,7 @@
 #include "Misc/Guid.h"
 #include "Misc/Optional.h"
 #include "Misc/OutputDeviceArchiveWrapper.h"
+#include "Misc/OutputDeviceHelper.h"
 #include "Misc/OutputDeviceRedirector.h"
 #include "Misc/Parse.h"
 #include "Misc/Paths.h"
@@ -65,6 +66,51 @@ static TAutoConsoleVariable<int32> CVarTraceSnapshotOnCrashTailSize(
 	ECVF_ReadOnly
 );
 
+static bool GCrashContextFormatEarlyCrashLog = false;
+static FAutoConsoleVariableRef CVarCrashContextFormatEarlyCrashLog(
+	TEXT("CrashContext.FormatEarlyCrashLog"),
+	GCrashContextFormatEarlyCrashLog,
+	TEXT("If true, log lines written to the crash report log will include timestamps and category tags,	matching the format of the on-disk log.\n"
+		"Only affects crashes that occur before log output devices are fully initialized, i.e. before the log backlog is disabled. (default = false)"),
+	ECVF_Default
+);
+
+/**
+* Output device used to write the file log's backlog in the case of an early crash before the file log is ready.
+*/
+class FCrashBacklogOutputDevice : public FOutputDevice
+{
+public:
+	FCrashBacklogOutputDevice(FArchive* InArchive)
+		: LogAr(InArchive)
+	{
+		check(InArchive);
+	}
+
+	// FOutputDevice interface
+	CORE_API void Flush()
+	{
+		LogAr->Flush();
+	}
+
+	CORE_API void Serialize(const TCHAR* Data, ELogVerbosity::Type Verbosity, const class FName& Category)
+	{
+		Serialize(Data, Verbosity, Category, -1.0);
+	}
+
+	CORE_API void Serialize(const TCHAR* Data, ELogVerbosity::Type Verbosity, const class FName& Category, const double Time)
+	{
+		if (Verbosity != ELogVerbosity::SetColor)
+		{
+			FOutputDeviceHelper::FormatCastAndSerializeLine(*LogAr, Data, Verbosity, Category, Time, bSuppressEventTag, bAutoEmitLineTerminator);
+		}
+	}
+	// End of FOutputDevice interface
+
+private:
+	FArchive* LogAr;
+};
+
 /**
  * A function-like type that creates a TStringBuilder to xml-escape a string
  */
@@ -1626,7 +1672,9 @@ FString FGenericCrashContext::DumpLog(const FString& CrashFolderAbsolute)
 			}
 			else
 			{
-				FOutputDeviceArchiveWrapper Wrapper(LogFile.Get());
+				FCrashBacklogOutputDevice Wrapper(LogFile.Get());
+				Wrapper.SetSuppressEventTag(!GCrashContextFormatEarlyCrashLog);
+				TGuardValue GuardPrintLogTimes(GPrintLogTimes, ELogTimes::SinceGStartTime);
 				GLog->SerializeBacklogForCrash(&Wrapper);
 			}
 
