Fix container paths when running with target "My Mac (Designed for iPad)"

#rnx
#Synced-CL 52432871

diff --git a/Engine/Source/Runtime/Core/Private/IOS/IOSPlatformFile.cpp b/Engine/Source/Runtime/Core/Private/IOS/IOSPlatformFile.cpp
--- a/Engine/Source/Runtime/Core/Private/IOS/IOSPlatformFile.cpp
+++ b/Engine/Source/Runtime/Core/Private/IOS/IOSPlatformFile.cpp
@@ -1004,6 +1004,21 @@ FString FIOSPlatformFile::ConvertToPlatformPath(const FString& Filename, bool bF
 		return Result;
 	}
 
+#if !UE_BUILD_SHIPPING
+	// My Mac (Designed for iPad) target in xcode: paths already resolved to the container
+	// should also be skipped
+	static bool bIsIOSAppOnMac = [[NSProcessInfo processInfo] isiOSAppOnMac];
+	if (bIsIOSAppOnMac)
+	{
+		static FString DocumentsBase = FString([NSSearchPathForDirectoriesInDomains(NSDocumentDirectory, NSUserDomainMask, YES) objectAtIndex:0]);
+		static FString LibraryBase = FString([NSSearchPathForDirectoriesInDomains(NSLibraryDirectory, NSUserDomainMask, YES) objectAtIndex:0]);
+		if (Result.StartsWith(DocumentsBase) || Result.StartsWith(LibraryBase))
+		{
+			return Result;
+		}
+	}
+#endif
+
 	if (Result.StartsWith(TEXT("~/")))
 	{
 		static FString ReadPathBase = FString([[NSBundle mainBundle]bundlePath] );
