Move Network monitoring teardown to an earlier point as PlatformTeardown is called too late.

#Synced-CL 52294780

diff --git a/Engine/Source/Runtime/Core/Private/Windows/WindowsPlatformMisc.cpp b/Engine/Source/Runtime/Core/Private/Windows/WindowsPlatformMisc.cpp
--- a/Engine/Source/Runtime/Core/Private/Windows/WindowsPlatformMisc.cpp
+++ b/Engine/Source/Runtime/Core/Private/Windows/WindowsPlatformMisc.cpp
@@ -1100,7 +1100,7 @@ void FWindowsPlatformMisc::PlatformInit()
 
 void FWindowsPlatformMisc::PlatformTearDown()
 {
-	StopNetworkConnectionMonitoring();
+	// StopNetworkConnectionMonitoring is now called via FCoreDelegates::OnEnginePreExit as this function is called too late in the exit sequnce
 	FPlatformProcess::CeaseBeingFirstInstance();
 }
 
@@ -4849,6 +4849,10 @@ void FWindowsPlatformMisc::StartNetworkConnectionMonitoring()
 	const ENetworkConnectionType InitialConnectionType = QueryNetworkConnectionTypeFromNLM();
 	GCachedNetworkConnectionType = InitialConnectionType;
 	SetNetworkConnectionStatus(NetworkConnectionTypeToStatus(InitialConnectionType));
+
+	// Register cleanup on OnEnginePreExit - this fires before OnPreExit and OnExit.
+	// We must clean it up earlier than PlatofrmTearDown because its callbacks can crash at that point
+	FCoreDelegates::OnEnginePreExit.AddStatic(StopNetworkConnectionMonitoring);
 }
 
 void FWindowsPlatformMisc::StopNetworkConnectionMonitoring()
