* STATS: Deprecated the StatsFile functionality (including the following console commands: "stat startfile", "stat startfileraw", "stat stopfile", "stat stopfileraw" and "stat testfile".) The StatsFile functionality is now toggled off by UE_ENABLE_STATS_FILE_DEPRECATED_IN_5_8 define.
* UnrealFrontend: Deprecated "CONVERT" command (needs UE_ENABLE_STATS_FILE_DEPRECATED_IN_5_8).

#Synced-CL 52466930

diff --git a/Engine/Source/Programs/UnrealFrontend/Private/Commands/StatsConvertCommand.cpp b/Engine/Source/Programs/UnrealFrontend/Private/Commands/StatsConvertCommand.cpp
--- a/Engine/Source/Programs/UnrealFrontend/Private/Commands/StatsConvertCommand.cpp
+++ b/Engine/Source/Programs/UnrealFrontend/Private/Commands/StatsConvertCommand.cpp
@@ -1,6 +1,9 @@
 // Copyright Epic Games, Inc. All Rights Reserved.
 
 #include "StatsConvertCommand.h"
+
+#if STATS && UE_ENABLE_STATS_FILE_DEPRECATED_IN_5_8
+
 #include "Misc/CommandLine.h"
 #include "HAL/FileManager.h"
 #include "Templates/UniquePtr.h"
@@ -168,3 +171,5 @@ void FStatsConvertCommand::InternalRun()
 		}
 	}
 }
+
+#endif // STATS && UE_ENABLE_STATS_FILE_DEPRECATED_IN_5_8
diff --git a/Engine/Source/Programs/UnrealFrontend/Private/Commands/StatsConvertCommand.h b/Engine/Source/Programs/UnrealFrontend/Private/Commands/StatsConvertCommand.h
--- a/Engine/Source/Programs/UnrealFrontend/Private/Commands/StatsConvertCommand.h
+++ b/Engine/Source/Programs/UnrealFrontend/Private/Commands/StatsConvertCommand.h
@@ -6,6 +6,7 @@
 #include "Stats/StatsData.h"
 #include "Stats/StatsFile.h"
 
+#if STATS && UE_ENABLE_STATS_FILE_DEPRECATED_IN_5_8
 
 class FStatsConvertCommand
 {
@@ -23,3 +24,5 @@ protected:
 	void InternalRun();
 
 };
+
+#endif // STATS && UE_ENABLE_STATS_FILE_DEPRECATED_IN_5_8
\ No newline at end of file
diff --git a/Engine/Source/Programs/UnrealFrontend/Private/UnrealFrontendMain.cpp b/Engine/Source/Programs/UnrealFrontend/Private/UnrealFrontendMain.cpp
--- a/Engine/Source/Programs/UnrealFrontend/Private/UnrealFrontendMain.cpp
+++ b/Engine/Source/Programs/UnrealFrontend/Private/UnrealFrontendMain.cpp
@@ -72,12 +72,15 @@ int32 UnrealFrontendMain( const TCHAR* CommandLine )
 	// NOTE - Some commands may require extra command line parameters.
 	if (bRunCommand)
 	{
+#if STATS && UE_ENABLE_STATS_FILE_DEPRECATED_IN_5_8
 		if (Command.Equals(TEXT("CONVERT"), ESearchCase::IgnoreCase))
 		{
 			FStatsConvertCommand::Run();
 		}
+		else
+#endif // STATS && UE_ENABLE_STATS_FILE_DEPRECATED_IN_5_8
 		// The 'LAUNCHPROFILE' command also needs '-PROFILENAME="MY_PROFILE_NAME"' as a command line parameter.
-		else if (Command.Equals(TEXT("LAUNCHPROFILE"), ESearchCase::IgnoreCase))
+		if (Command.Equals(TEXT("LAUNCHPROFILE"), ESearchCase::IgnoreCase))
 		{
 			FLaunchFromProfileCommand* ProfileLaunch = new FLaunchFromProfileCommand;
 			ProfileLaunch->Run(Params);
diff --git a/Engine/Source/Runtime/Core/Private/Stats/StatsCommand.cpp b/Engine/Source/Runtime/Core/Private/Stats/StatsCommand.cpp
--- a/Engine/Source/Runtime/Core/Private/Stats/StatsCommand.cpp
+++ b/Engine/Source/Runtime/Core/Private/Stats/StatsCommand.cpp
@@ -1898,21 +1898,18 @@ static void PrintStatsHelpToOutputDevice( FOutputDevice& Ar )
 	Ar.Log( TEXT( "	Changes stats rendering display options" ) );
 #endif // WITH_ENGINE
 
+#if UE_ENABLE_STATS_FILE_DEPRECATED_IN_5_8
 	Ar.Log( TEXT("stat startfile - starts dumping a capture"));
 	Ar.Log( TEXT("stat stopfile - stops dumping a capture (regular, raw, memory)"));
-
 	Ar.Log( TEXT("stat startfileraw - starts dumping a raw capture"));
+	Ar.Log( TEXT("stat testfile - loads the last saved capture and dumps first, middle and last frame"));
+#endif // UE_ENABLE_STATS_FILE_DEPRECATED_IN_5_8
 
 	Ar.Log( TEXT("stat toggledebug - toggles tracking the most memory expensive stats"));
 
 	Ar.Log( TEXT( "stat slow [-ms=1.0] [-depth=4] - toggles displaying the game and render thread stats" ) );
 
-	Ar.Log( TEXT("add -memoryprofiler in the command line to enable the memory profiling"));
-	Ar.Log( TEXT("stat stopfile - stops tracking all memory operations and writes the results to the file"));
-
 	Ar.Log( TEXT("stat namedmarker #markername# - adds a custom marker to the stats stream"));
-
-	Ar.Log( TEXT( "stat testfile - loads the last saved capture and dumps first, middle and last frame" ) );
 }
 
 #endif
@@ -2049,6 +2046,7 @@ static void StatCmd(FString InCmd, bool bStatCommand, FOutputDevice* Ar /*= null
 			StatsPrimaryEnableAdd();
 			DumpEventsDelegateHandle = Stats.NewFrameDelegate.AddStatic(&DumpEventsOnce, DumpEventsCullMS, bDisplayAllThreads);
 		}
+#if UE_ENABLE_STATS_FILE_DEPRECATED_IN_5_8
 		else if (FParse::Command(&Cmd, TEXT("STARTFILE")))
 		{
 			FString Filename;
@@ -2091,6 +2089,7 @@ static void StatCmd(FString InCmd, bool bStatCommand, FOutputDevice* Ar /*= null
 		{
 			FCommandStatsFile::Get().TestLastSaved();
 		}
+#endif // UE_ENABLE_STATS_FILE_DEPRECATED_IN_5_8
 		else if (FParse::Command(&Cmd, TEXT("testdisable")))
 		{
 			FThreadStats::PrimaryDisableForever();
@@ -2293,6 +2292,7 @@ namespace UE::Stats
 				{
 					PrintStatsHelpToOutputDevice(*Ar);
 				}
+#if UE_ENABLE_STATS_FILE_DEPRECATED_IN_5_8
 				else if (FParse::Command(&TempCmd, TEXT("STARTFILE")))
 				{
 					FString Filename;
@@ -2311,6 +2311,7 @@ namespace UE::Stats
 					AddArgs += TEXT(" ");
 					AddArgs += CreateProfileFilename(FStatConstants::StatsFileRawExtension, true);
 				}
+#endif // UE_ENABLE_STATS_FILE_DEPRECATED_IN_5_8
 				else if (FParse::Command(&TempCmd, TEXT("DUMPFRAME")))
 				{
 				}
@@ -2341,12 +2342,14 @@ namespace UE::Stats
 				else if (FParse::Command(&TempCmd, TEXT("DumpEvents")))
 				{
 				}
+#if UE_ENABLE_STATS_FILE_DEPRECATED_IN_5_8
 				else if (FParse::Command(&TempCmd, TEXT("STOPFILE")))
 				{
 				}
 				else if (FParse::Command(&TempCmd, TEXT("TESTFILE")))
 				{
 				}
+#endif // UE_ENABLE_STATS_FILE_DEPRECATED_IN_5_8
 				else if (FParse::Command(&TempCmd, TEXT("testdisable")))
 				{
 				}
diff --git a/Engine/Source/Runtime/Core/Private/Stats/StatsFile.cpp b/Engine/Source/Runtime/Core/Private/Stats/StatsFile.cpp
--- a/Engine/Source/Runtime/Core/Private/Stats/StatsFile.cpp
+++ b/Engine/Source/Runtime/Core/Private/Stats/StatsFile.cpp
@@ -5,6 +5,9 @@
 =============================================================================*/
 
 #include "Stats/StatsFile.h"
+
+#if STATS && UE_ENABLE_STATS_FILE_DEPRECATED_IN_5_8
+
 #include "HAL/FileManager.h"
 #include "Misc/Paths.h"
 #include "Internationalization/Internationalization.h"
@@ -12,8 +15,6 @@
 #include "Serialization/MemoryReader.h"
 #include "Templates/UniquePtr.h"
 
-#if	STATS
-
 #include "Misc/ScopeExit.h"
 
 DECLARE_CYCLE_STAT( TEXT( "Stream File" ), STAT_StreamFile, STATGROUP_StatSystem );
@@ -1163,4 +1164,4 @@ void FCommandStatsFile::TestLastSaved()
 
 #undef LOCTEXT_NAMESPACE
 
-#endif // STATS
+#endif // STATS && UE_ENABLE_STATS_FILE_DEPRECATED_IN_5_8
diff --git a/Engine/Source/Runtime/Core/Public/Stats/StatsFile.h b/Engine/Source/Runtime/Core/Public/Stats/StatsFile.h
--- a/Engine/Source/Runtime/Core/Public/Stats/StatsFile.h
+++ b/Engine/Source/Runtime/Core/Public/Stats/StatsFile.h
@@ -2,6 +2,14 @@
 
 #pragma once
 
+#include "Misc/Build.h"
+
+#ifndef UE_ENABLE_STATS_FILE_DEPRECATED_IN_5_8
+#define UE_ENABLE_STATS_FILE_DEPRECATED_IN_5_8 0
+#endif
+
+#if UE_ENABLE_INCLUDE_ORDER_DEPRECATED_IN_5_8 || UE_ENABLE_STATS_FILE_DEPRECATED_IN_5_8
+
 #include "Async/AsyncWork.h"
 #include "Containers/Array.h"
 #include "Containers/ChunkedArray.h"
@@ -29,10 +37,9 @@
 #include "UObject/NameTypes.h"
 #include "UObject/UnrealNames.h"
 
-class FAsyncStatsWrite;
-struct FStatsReadFile;
+#endif // UE_ENABLE_INCLUDE_ORDER_DEPRECATED_IN_5_8 || UE_ENABLE_STATS_FILE_DEPRECATED_IN_5_8
 
-#if	STATS
+#if STATS && UE_ENABLE_STATS_FILE_DEPRECATED_IN_5_8
 
 class FAsyncStatsWrite;
 struct FStatsReadFile;
@@ -1252,4 +1259,4 @@ protected:
 };
 
 
-#endif // STATS
+#endif // STATS && UE_ENABLE_STATS_FILE_DEPRECATED_IN_5_8
diff --git a/Engine/Source/Runtime/Engine/Private/UnrealEngine.cpp b/Engine/Source/Runtime/Engine/Private/UnrealEngine.cpp
--- a/Engine/Source/Runtime/Engine/Private/UnrealEngine.cpp
+++ b/Engine/Source/Runtime/Engine/Private/UnrealEngine.cpp
@@ -30,6 +30,8 @@ UnrealEngine.cpp: Implements the UEngine class and helpers.
 #include "Misc/OutputDeviceFile.h"
 #include "Math/ColorList.h"
 #include "StateStreamManager.h"
+#include "Stats/StatsData.h"
+#include "Stats/StatsFile.h"
 #include "Stats/StatsMisc.h"
 #include "Stats/ThreadIdleStats.h"
 #include "Misc/AutomationTest.h"
@@ -109,7 +111,6 @@ UnrealEngine.cpp: Implements the UEngine class and helpers.
 #include "IHeadMountedDisplay.h"
 #include "IXRTrackingSystem.h"
 #include "IXRLoadingScreen.h"
-#include "Stats/StatsFile.h"
 #include "Audio/AudioDebug.h"
 #include "AudioDevice.h"
 #include "Animation/SkeletalMeshActor.h"
@@ -13607,6 +13608,7 @@ void DrawStatsHUD( UWorld* World, FViewport* Viewport, FCanvas* Canvas, UCanvas*
 				MessageY += FontSizeY;
 			}
 
+#if UE_ENABLE_STATS_FILE_DEPRECATED_IN_5_8
 			const bool bIsStatsFileActive = FCommandStatsFile::Get().IsStatFileActive();
 			if (bIsStatsFileActive)
 			{
@@ -13615,6 +13617,7 @@ void DrawStatsHUD( UWorld* World, FViewport* Viewport, FCanvas* Canvas, UCanvas*
 				Canvas->DrawItem(SmallTextItem, FVector2D(MessageX, MessageY));
 				MessageY += FontSizeY;
 			}
+#endif // UE_ENABLE_STATS_FILE_DEPRECATED_IN_5_8
 
 #if ENABLE_LOW_LEVEL_MEM_TRACKER
 			if (FLowLevelMemTracker::Get().IsEnabled() && !FPlatformMemory::IsExtraDevelopmentMemoryAvailable())
