Fixed -nopch compile error in StatsFile.h

#rnx
#Synced-CL 52466930

diff --git a/Engine/Source/Runtime/Core/Public/Stats/StatsFile.h b/Engine/Source/Runtime/Core/Public/Stats/StatsFile.h
--- a/Engine/Source/Runtime/Core/Public/Stats/StatsFile.h
+++ b/Engine/Source/Runtime/Core/Public/Stats/StatsFile.h
@@ -3,6 +3,7 @@
 #pragma once
 
 #include "Misc/Build.h"
+#include "Misc/CoreMiscDefines.h"
 
 #ifndef UE_ENABLE_STATS_FILE_DEPRECATED_IN_5_8
 #define UE_ENABLE_STATS_FILE_DEPRECATED_IN_5_8 0
