Add UE_PLATFORM_* definitions to mirror existing PLATFORM_* definitions

#Synced-CL 52476816

diff --git a/Engine/Source/Programs/UnrealBuildTool/Platform/Android/UEBuildAndroid.cs b/Engine/Source/Programs/UnrealBuildTool/Platform/Android/UEBuildAndroid.cs
--- a/Engine/Source/Programs/UnrealBuildTool/Platform/Android/UEBuildAndroid.cs
+++ b/Engine/Source/Programs/UnrealBuildTool/Platform/Android/UEBuildAndroid.cs
@@ -487,6 +487,7 @@ namespace UnrealBuildTool
 			CompileEnvironment.Definitions.Add("_UNICODE");
 
 			CompileEnvironment.Definitions.Add("PLATFORM_ANDROID=1");
+			CompileEnvironment.Definitions.Add("UE_PLATFORM_ANDROID=1");
 			CompileEnvironment.Definitions.Add("ANDROID=1");
 
 			CompileEnvironment.Definitions.Add("WITH_EDITOR=0");
@@ -531,6 +532,7 @@ namespace UnrealBuildTool
 			CompileEnvironment.Definitions.Add("_UNICODE");
 
 			CompileEnvironment.Definitions.Add("PLATFORM_ANDROID=1");
+			CompileEnvironment.Definitions.Add("UE_PLATFORM_ANDROID=1");
 			CompileEnvironment.Definitions.Add("ANDROID=1");
 
 			CompileEnvironment.Definitions.Add("WITH_EDITOR=0");
diff --git a/Engine/Source/Programs/UnrealBuildTool/Platform/Apple/UEBuildApple.cs b/Engine/Source/Programs/UnrealBuildTool/Platform/Apple/UEBuildApple.cs
--- a/Engine/Source/Programs/UnrealBuildTool/Platform/Apple/UEBuildApple.cs
+++ b/Engine/Source/Programs/UnrealBuildTool/Platform/Apple/UEBuildApple.cs
@@ -236,6 +236,7 @@ namespace UnrealBuildTool
 			// shared stuff for all apple platforms
 			
 			CompileEnvironment.Definitions.Add("PLATFORM_APPLE=1");
+			CompileEnvironment.Definitions.Add("UE_PLATFORM_APPLE=1");
 
 			AppleExports.GetSwiftIntegrationSettings(Target.ProjectFile, Target.Type, Platform, out bool bUseSwiftUIMain, out _);
 			CompileEnvironment.Definitions.Add("UE_USE_SWIFT_UI_MAIN=" + (bUseSwiftUIMain ? "1" : "0"));
diff --git a/Engine/Source/Programs/UnrealBuildTool/Platform/IOS/UEBuildIOS.cs b/Engine/Source/Programs/UnrealBuildTool/Platform/IOS/UEBuildIOS.cs
--- a/Engine/Source/Programs/UnrealBuildTool/Platform/IOS/UEBuildIOS.cs
+++ b/Engine/Source/Programs/UnrealBuildTool/Platform/IOS/UEBuildIOS.cs
@@ -1074,6 +1074,7 @@ namespace UnrealBuildTool
 			}
 
 			CompileEnvironment.Definitions.Add("PLATFORM_IOS=1");
+			CompileEnvironment.Definitions.Add("UE_PLATFORM_IOS=1");
 			// set a define that is _only_ on iPhones, and not AppleTV, etc
 			if (Target.Platform == UnrealTargetPlatform.IOS)
 			{
diff --git a/Engine/Source/Programs/UnrealBuildTool/Platform/Linux/UEBuildLinux.cs b/Engine/Source/Programs/UnrealBuildTool/Platform/Linux/UEBuildLinux.cs
--- a/Engine/Source/Programs/UnrealBuildTool/Platform/Linux/UEBuildLinux.cs
+++ b/Engine/Source/Programs/UnrealBuildTool/Platform/Linux/UEBuildLinux.cs
@@ -518,6 +518,8 @@ namespace UnrealBuildTool
 		{
 			CompileEnvironment.Definitions.Add("PLATFORM_LINUX=1");
 			CompileEnvironment.Definitions.Add("PLATFORM_UNIX=1");
+			CompileEnvironment.Definitions.Add("UE_PLATFORM_LINUX=1");
+			CompileEnvironment.Definitions.Add("UE_PLATFORM_UNIX=1");
 
 			CompileEnvironment.Definitions.Add("LINUX=1"); // For libOGG
 
diff --git a/Engine/Source/Programs/UnrealBuildTool/Platform/Mac/UEBuildMac.cs b/Engine/Source/Programs/UnrealBuildTool/Platform/Mac/UEBuildMac.cs
--- a/Engine/Source/Programs/UnrealBuildTool/Platform/Mac/UEBuildMac.cs
+++ b/Engine/Source/Programs/UnrealBuildTool/Platform/Mac/UEBuildMac.cs
@@ -462,6 +462,7 @@ namespace UnrealBuildTool
 			base.SetUpEnvironment(Target, CompileEnvironment, LinkEnvironment);
 
 			CompileEnvironment.Definitions.Add("PLATFORM_MAC=1");
+			CompileEnvironment.Definitions.Add("UE_PLATFORM_MAC=1");
 
 			if (Target.MacPlatform.bEnableRayTracing && Target.Type != TargetType.Server)
 			{
diff --git a/Engine/Source/Programs/UnrealBuildTool/Platform/TVOS/UEBuildTVOS.cs b/Engine/Source/Programs/UnrealBuildTool/Platform/TVOS/UEBuildTVOS.cs
--- a/Engine/Source/Programs/UnrealBuildTool/Platform/TVOS/UEBuildTVOS.cs
+++ b/Engine/Source/Programs/UnrealBuildTool/Platform/TVOS/UEBuildTVOS.cs
@@ -97,6 +97,7 @@ namespace UnrealBuildTool
 		{
 			base.SetUpEnvironment(Target, CompileEnvironment, LinkEnvironment);
 			CompileEnvironment.Definitions.Add("PLATFORM_TVOS=1");
+			CompileEnvironment.Definitions.Add("UE_PLATFORM_TVOS=1");
 
 			// TVOS uses only IOS header files, so use it's platform headers
 			CompileEnvironment.Definitions.Add("OVERRIDE_PLATFORM_HEADER_NAME=IOS");
diff --git a/Engine/Source/Programs/UnrealBuildTool/Platform/Windows/UEBuildWindows.cs b/Engine/Source/Programs/UnrealBuildTool/Platform/Windows/UEBuildWindows.cs
--- a/Engine/Source/Programs/UnrealBuildTool/Platform/Windows/UEBuildWindows.cs
+++ b/Engine/Source/Programs/UnrealBuildTool/Platform/Windows/UEBuildWindows.cs
@@ -2014,6 +2014,8 @@ namespace UnrealBuildTool
 
 			CompileEnvironment.Definitions.Add("PLATFORM_WINDOWS=1");
 			CompileEnvironment.Definitions.Add("PLATFORM_MICROSOFT=1");
+			CompileEnvironment.Definitions.Add("UE_PLATFORM_WINDOWS=1");
+			CompileEnvironment.Definitions.Add("UE_PLATFORM_MICROSOFT=1");
 
 			string? OverridePlatformHeaderName = GetOverridePlatformHeaderName();
 			if (!String.IsNullOrEmpty(OverridePlatformHeaderName))
diff --git a/Engine/Source/Runtime/Core/Public/HAL/Platform.h b/Engine/Source/Runtime/Core/Public/HAL/Platform.h
--- a/Engine/Source/Runtime/Core/Public/HAL/Platform.h
+++ b/Engine/Source/Runtime/Core/Public/HAL/Platform.h
@@ -67,6 +67,34 @@
 	#define PLATFORM_MICROSOFT 0
 #endif
 
+#if !defined(UE_PLATFORM_WINDOWS)
+	#define UE_PLATFORM_WINDOWS 0
+#endif
+#if !defined(UE_PLATFORM_MAC)
+	#define UE_PLATFORM_MAC 0
+#endif
+#if !defined(UE_PLATFORM_IOS)
+	#define UE_PLATFORM_IOS 0
+#endif
+#if !defined(UE_PLATFORM_TVOS)
+	#define UE_PLATFORM_TVOS 0
+#endif
+#if !defined(UE_PLATFORM_ANDROID)
+	#define UE_PLATFORM_ANDROID 0
+#endif
+#if !defined(UE_PLATFORM_APPLE)
+	#define UE_PLATFORM_APPLE 0
+#endif
+#if !defined(UE_PLATFORM_LINUX)
+	#define UE_PLATFORM_LINUX 0
+#endif
+#if !defined(UE_PLATFORM_UNIX)
+	#define UE_PLATFORM_UNIX 0
+#endif
+#if !defined(UE_PLATFORM_MICROSOFT)
+	#define UE_PLATFORM_MICROSOFT 0
+#endif
+
 // Platform specific compiler pre-setup.
 #include "PreprocessorHelpers.h"
 #include COMPILED_PLATFORM_HEADER(PlatformCompilerPreSetup.h)
