Fix logging for ios 26 since NSLog is now automatically redacted in the console app

#rnx

#Synced-CL 52432871

diff --git a/Engine/Source/Runtime/Core/Private/Apple/ApplePlatformMisc.cpp b/Engine/Source/Runtime/Core/Private/Apple/ApplePlatformMisc.cpp
--- a/Engine/Source/Runtime/Core/Private/Apple/ApplePlatformMisc.cpp
+++ b/Engine/Source/Runtime/Core/Private/Apple/ApplePlatformMisc.cpp
@@ -74,8 +74,9 @@ FString FApplePlatformMisc::GetEnvironmentVariable(const TCHAR* VariableName)
 void FApplePlatformMisc::LocalPrint(const TCHAR* Message)
 {
 #if !UE_BUILD_SHIPPING
-	// This goes to all output consoles, not only system one, stdout included.
-	NSLog(@"[UE] %s", TCHAR_TO_UTF8(Message));
+	// Use os_log with %{public}s so messages are visible in Console.app
+	// without Xcode attached. NSLog is redacted on iOS 26+ unified logging.
+	os_log(OS_LOG_DEFAULT, "[UE] %{public}s", TCHAR_TO_UTF8(Message));
 #else
 	// This only goes to the system console.
 	// Private by default, would be visible in Xcode.
