Adding support for the CREATE_BREAKAWAY_FROM_JOB flag on Windows processes to accomodate a requirement for spawning EOSH dev sessions.

#Synced-CL 52463319

diff --git a/Engine/Source/Runtime/Core/Private/Windows/WindowsPlatformProcess.cpp b/Engine/Source/Runtime/Core/Private/Windows/WindowsPlatformProcess.cpp
--- a/Engine/Source/Runtime/Core/Private/Windows/WindowsPlatformProcess.cpp
+++ b/Engine/Source/Runtime/Core/Private/Windows/WindowsPlatformProcess.cpp
@@ -504,7 +504,7 @@ FProcHandle FWindowsPlatformProcess::CreateProc(const TCHAR* URL, const TCHAR* P
 	}
 	// CreateProc used to only have a single "write" pipe argument that both stdout and stderr would be piped into on Windows,
 	// so for this overload we'll preserve that behaviour for compatibility with existing code
-	return CreateProcInternal(URL, Parms, bLaunchDetached, bLaunchHidden, bLaunchReallyHidden, OutProcessID, PriorityModifier, OptionalWorkingDirectory, PipeWriteChild, PipeReadChild, PipeWriteChild, false);
+	return CreateProcInternal(URL, Parms, bLaunchDetached, bLaunchHidden, bLaunchReallyHidden, false, OutProcessID, PriorityModifier, OptionalWorkingDirectory, PipeWriteChild, PipeReadChild, PipeWriteChild, false);
 }
 
 FProcHandle FWindowsPlatformProcess::CreateProc(const TCHAR* URL, const TCHAR* Parms, bool bLaunchDetached, bool bLaunchHidden, bool bLaunchReallyHidden, uint32* OutProcessID, int32 PriorityModifier, const TCHAR* OptionalWorkingDirectory, void* PipeWriteChild, void* PipeReadChild, void* PipeStdErrChild)
@@ -539,18 +539,18 @@ FProcHandle FWindowsPlatformProcess::CreateProc(const TCHAR* URL, const TCHAR* P
 			PipeStdErrChild = nullptr;
 		}
 	}
-	return CreateProcInternal(URL, Parms, bLaunchDetached, bLaunchHidden, bLaunchReallyHidden, OutProcessID, PriorityModifier, OptionalWorkingDirectory, PipeWriteChild, PipeReadChild, PipeStdErrChild, false);
+	return CreateProcInternal(URL, Parms, bLaunchDetached, bLaunchHidden, bLaunchReallyHidden, false, OutProcessID, PriorityModifier, OptionalWorkingDirectory, PipeWriteChild, PipeReadChild, PipeStdErrChild, false);
 }
 
 TTuple<FProcHandle, UE::HAL::EProcessId> FWindowsPlatformProcess::CreateProc(const UE::HAL::FProcessStartInfo& StartInfo)
 {
 	uint32 Id;
-	FProcHandle Handle = CreateProcInternal(StartInfo.Uri, StartInfo.Arguments, StartInfo.bDetached, false, StartInfo.bHidden, &Id, StartInfo.PriorityModifier, StartInfo.WorkingDirectory, StartInfo.StdOut, StartInfo.StdIn, StartInfo.StdErr, StartInfo.bInheritHandles);
+	FProcHandle Handle = CreateProcInternal(StartInfo.Uri, StartInfo.Arguments, StartInfo.bDetached, false, StartInfo.bHidden, StartInfo.bBreakawayFromJob, &Id, StartInfo.PriorityModifier, StartInfo.WorkingDirectory, StartInfo.StdOut, StartInfo.StdIn, StartInfo.StdErr, StartInfo.bInheritHandles);
 	return {Handle, UE::HAL::EProcessId{Id}};
 }
 
 #include "Windows/AllowWindowsPlatformTypes.h"
-FProcHandle FWindowsPlatformProcess::CreateProcInternal(const TCHAR* URL, const TCHAR* Parms, bool bLaunchDetached, bool bLaunchHidden, bool bLaunchReallyHidden, uint32* OutProcessID, int32 PriorityModifier, const TCHAR* OptionalWorkingDirectory, void* PipeWriteChild, void* PipeReadChild, void* PipeStdErrChild, bool bInheritHandles)
+FProcHandle FWindowsPlatformProcess::CreateProcInternal(const TCHAR* URL, const TCHAR* Parms, bool bLaunchDetached, bool bLaunchHidden, bool bLaunchReallyHidden, bool bBreakawayFromJob, uint32* OutProcessID, int32 PriorityModifier, const TCHAR* OptionalWorkingDirectory, void* PipeWriteChild, void* PipeReadChild, void* PipeStdErrChild, bool bInheritHandles)
 {
 	check(URL);
 	check(URL[0] != '\"');
@@ -602,6 +602,11 @@ FProcHandle FWindowsPlatformProcess::CreateProcInternal(const TCHAR* URL, const
 		CreationFlags |= bLaunchReallyHidden ? CREATE_NO_WINDOW : DETACHED_PROCESS;
 	}
 
+	if (bBreakawayFromJob)
+	{
+		CreationFlags |= CREATE_BREAKAWAY_FROM_JOB;
+	}
+
 	DWORD Flags;
 	WORD ShowWindow;
 	if (bLaunchReallyHidden)
diff --git a/Engine/Source/Runtime/Core/Public/GenericPlatform/GenericPlatformProcess.h b/Engine/Source/Runtime/Core/Public/GenericPlatform/GenericPlatformProcess.h
--- a/Engine/Source/Runtime/Core/Public/GenericPlatform/GenericPlatformProcess.h
+++ b/Engine/Source/Runtime/Core/Public/GenericPlatform/GenericPlatformProcess.h
@@ -202,6 +202,7 @@ namespace UE::HAL
 		bool bDetached = false;
 		bool bHidden = false;
 		bool bInheritHandles = false;
+		bool bBreakawayFromJob = false;
 		FReadHandle StdIn;
 		FWriteHandle StdOut;
 		FWriteHandle StdErr = StdOut;
diff --git a/Engine/Source/Runtime/Core/Public/Windows/WindowsPlatformProcess.h b/Engine/Source/Runtime/Core/Public/Windows/WindowsPlatformProcess.h
--- a/Engine/Source/Runtime/Core/Public/Windows/WindowsPlatformProcess.h
+++ b/Engine/Source/Runtime/Core/Public/Windows/WindowsPlatformProcess.h
@@ -246,7 +246,7 @@ protected:
 
 private:
 
-	static FProcHandle CreateProcInternal(const TCHAR* URL, const TCHAR* Parms, bool bLaunchDetached, bool bLaunchHidden, bool bLaunchReallyHidden, uint32* OutProcessID, int32 PriorityModifier, const TCHAR* OptionalWorkingDirectory, void* PipeWriteChild, void* PipeReadChild, void* PipeStdErrChild, bool bInheritHandles);
+	static FProcHandle CreateProcInternal(const TCHAR* URL, const TCHAR* Parms, bool bLaunchDetached, bool bLaunchHidden, bool bLaunchReallyHidden, bool bBreakawayFromJob, uint32* OutProcessID, int32 PriorityModifier, const TCHAR* OptionalWorkingDirectory, void* PipeWriteChild, void* PipeReadChild, void* PipeStdErrChild, bool bInheritHandles);
 	
 	static void LaunchURLInternal(const TCHAR* URL, const TCHAR* Parms, FString* Error);
 
