[AutoRTFM] Fix retry behavior of transactionally-safe recursive mutexes.

Previously, we would leak a pooled mutex state if a retry occurred. This refactor also happens to improve benchmarks by about 1-3%.

#Synced-CL 52411280

diff --git a/Engine/Source/Programs/AutoRTFMTests/Private/MutexTests.cpp b/Engine/Source/Programs/AutoRTFMTests/Private/MutexTests.cpp
--- a/Engine/Source/Programs/AutoRTFMTests/Private/MutexTests.cpp
+++ b/Engine/Source/Programs/AutoRTFMTests/Private/MutexTests.cpp
@@ -1396,6 +1396,35 @@ TEMPLATE_TEST_CASE("MutexTests", "",
 
 	REQUIRE(NumLocksUsed == UE::Private::TransactionallySafeMutexPool::GetNumUsed());
 
+	SECTION("Retry Releases Locks")
+	{
+		MutexType Mutex;
+		bool bAlreadyRetried = false;
+
+		AutoRTFM::Testing::Commit([&]
+		{
+			// Mutex must start off unlocked.
+			RequireNotLocked(Mutex);
+			UE::TScopeLock Lock(Mutex);
+			RequireLocked(Mutex);
+
+			if (!bAlreadyRetried)
+			{
+				AutoRTFM::OnRetry([&]
+				{
+					// Mutex must be unlocked when we reach OnRetry, as OnRetry happens after all OnAborts.
+					RequireNotLocked(Mutex);
+					bAlreadyRetried = true;
+				});
+
+				// Trigger a retry with the lock held.
+				AutoRTFM::CascadingRetryTransaction();
+			}
+		});
+	}
+
+	REQUIRE(NumLocksUsed == UE::Private::TransactionallySafeMutexPool::GetNumUsed());
+
 	if constexpr (bIsShared && bIsRecursive)
 	{
 		SECTION("Transact(SharedLock(Transact(SharedLock(), Abort)))")
diff --git a/Engine/Source/Runtime/Core/Public/Async/TransactionallySafeRecursiveMutex.h b/Engine/Source/Runtime/Core/Public/Async/TransactionallySafeRecursiveMutex.h
--- a/Engine/Source/Runtime/Core/Public/Async/TransactionallySafeRecursiveMutex.h
+++ b/Engine/Source/Runtime/Core/Public/Async/TransactionallySafeRecursiveMutex.h
@@ -59,11 +59,12 @@ struct FTransactionallySafeRecursiveMutexDefinition final
 		}
 		else if (AutoRTFM::IsClosed())
 		{
+			// We explicitly extend the lifetime of `State` here, in case `this` lives on the
+			// stack and ends up being overwritten before the on-abort is hit.
+			State->Acquire();
+
 			UE_AUTORTFM_OPEN
 			{
-				// We explicitly extend the lifetime of `State` here, in case `this` lives on the
-				// stack and ends up being overwritten before the on-abort is hit.
-				State->Acquire();
 				State->LockTransactionalLock();
 			};
 
@@ -73,7 +74,7 @@ struct FTransactionallySafeRecursiveMutexDefinition final
 				State->UnlockTransactionalLock();
 			});
 
-			State->ReleaseOnComplete();
+			State->ReleaseOnCommit();
 		}
 		else
 		{
@@ -91,20 +92,16 @@ struct FTransactionallySafeRecursiveMutexDefinition final
 		}
 		else if (AutoRTFM::IsClosed())
 		{
-			UE_AUTORTFM_OPEN
-			{
-				// We explicitly extend the lifetime of `State` here, in case `this` lives on the
-				// stack and ends up being overwritten before the on-commit is hit.
-				State->Acquire();
-			};
+			// We explicitly extend the lifetime of `State` here, in case `this` lives on the
+			// stack and ends up being overwritten before the on-commit is hit.
+			State->Acquire();
 
 			AutoRTFM::OnCommit([State = this->State]
 			{
 				// We don't want to actually unlock the transactional mutex until the transaction reaches commit.
 				State->UnlockTransactionalLock();
+				State->ReleaseImmediately();
 			});
-
-			State->ReleaseOnComplete();
 		}
 		else
 		{
@@ -166,28 +163,11 @@ private:
 			}
 		}
 
-		AUTORTFM_DISABLE void Acquire()
+		AUTORTFM_ENABLE void Acquire()
 		{
 			RefCount += 1;
 		}
 
-		void ReleaseOnComplete()
-		{
-			checkSlow(RefCount > 0);
-
-			if (AutoRTFM::IsClosed())
-			{
-				AutoRTFM::OnComplete([this]
-				{
-					ReleaseImmediately();
-				});
-			}
-			else
-			{
-				ReleaseImmediately();
-			}
-		}
-
 		void ReleaseOnCommit()
 		{
 			checkSlow(RefCount > 0);
