- Added the concept of auto-loading non-plugin modules based on .ini settings (see [AutoLoadModulesByPhase] in BaseEngine.ini)
- Making use of it to load TargetPlatformManagerModule earlier to kick off the async UBT process (as thought of by [REDACTED_PERSON])

#Synced-CL 52482450

diff --git a/Engine/Config/BaseEngine.ini b/Engine/Config/BaseEngine.ini
--- a/Engine/Config/BaseEngine.ini
+++ b/Engine/Config/BaseEngine.ini
@@ -30,6 +30,10 @@ Name=Player
 SaveExt=usa
 Port=7777
 
+[AutoLoadModulesByPhase]
+; matching phase with FDelayedAutoRegisterHelper GSDKInit in TargetPlatformManagerModule.cpp
++FileSystemReady=TargetPlatform
+
 [HTTP]
 HttpConnectionTimeout=30
 HttpActivityTimeout=30
diff --git a/Engine/Source/Developer/TargetPlatform/Private/TargetPlatformManagerModule.cpp b/Engine/Source/Developer/TargetPlatform/Private/TargetPlatformManagerModule.cpp
--- a/Engine/Source/Developer/TargetPlatform/Private/TargetPlatformManagerModule.cpp
+++ b/Engine/Source/Developer/TargetPlatform/Private/TargetPlatformManagerModule.cpp
@@ -596,8 +596,28 @@ public:
 		, bHasInitErrors(false)
 		, bIgnoreFirstDelegateCall(true)
 	{
-#if WITH_EDITOR && UE_WITH_TURNKEY_SUPPORT
+	}
+
+	/** Destructor. */
+	virtual ~FTargetPlatformManagerModule()
+	{
+		if (bIsInitialized)
+		{
+			FModuleManager::Get().OnModulesChanged().RemoveAll(this);
+		}
+	}
+
+
+	virtual void Initialize()
+	{
+		if (bIsInitialized)
+		{
+			return;
+		}
 
+		bIsInitialized = true;
+
+#if WITH_EDITOR && UE_WITH_TURNKEY_SUPPORT
 		ITurnkeySupportModule::Get().UpdateSdkInfo();
 #endif
 
@@ -605,19 +625,11 @@ public:
 
 		TextureFormatManager = FModuleManager::LoadModulePtr<ITextureFormatManagerModule>("TextureFormat");
 
-		// Calling a virtual function from a constructor, but with no expectation that a derived implementation of this
-		// method would be called.  This is solely to avoid duplicating code in this implementation, not for polymorphism.
-		FTargetPlatformManagerModule::Invalidate();
+		Invalidate();
 
 		FModuleManager::Get().OnModulesChanged().AddRaw(this, &FTargetPlatformManagerModule::ModulesChangesCallback);
 	}
 
-	/** Destructor. */
-	virtual ~FTargetPlatformManagerModule()
-	{
-		FModuleManager::Get().OnModulesChanged().RemoveAll(this);
-	}
-
 public:
 
 	// ITargetPlatformManagerModule interface
@@ -1758,6 +1770,8 @@ private:
 	// TargetPlatform= on the command line.
 	bool bRestrictFormatsToRuntimeOnly;
 
+	bool bIsInitialized = false;
+
 	// Flag to force reinitialization of all cached data. This is needed to have up-to-date caches
 	// in case of a module reload of a TargetPlatform-Module.
 	bool bForceCacheUpdate;
diff --git a/Engine/Source/Developer/TargetPlatform/Public/Interfaces/ITargetPlatformManagerModule.h b/Engine/Source/Developer/TargetPlatform/Public/Interfaces/ITargetPlatformManagerModule.h
--- a/Engine/Source/Developer/TargetPlatform/Public/Interfaces/ITargetPlatformManagerModule.h
+++ b/Engine/Source/Developer/TargetPlatform/Public/Interfaces/ITargetPlatformManagerModule.h
@@ -3,10 +3,10 @@
 #pragma once
 
 #include "CoreMinimal.h"
-#include "Containers/StringView.h"
 #include "Delegates/Delegate.h"
 #include "Interfaces/ITargetDevice.h"
 #include "Modules/ModuleInterface.h"
+#include "Modules/ModuleManager.h"
 
 class ITargetPlatform;
 class ITargetPlatformControls;
@@ -22,6 +22,16 @@ class ITargetPlatformManagerModule
 {
 public:
 
+	static FORCEINLINE ITargetPlatformManagerModule* Get()
+	{
+		ITargetPlatformManagerModule* Module = FModuleManager::LoadModulePtr<ITargetPlatformManagerModule>("TargetPlatform");
+		if (Module)
+		{
+			Module->Initialize();
+		}
+		return Module;
+	}
+
 	/**
 	 * Finds an audio format with the specified name.
 	 *
@@ -206,8 +216,16 @@ public:
 	 */
 	virtual bool UpdateAfterSDKInstall(FName TargetPlatformName) = 0;
 
-public:
 
 	/** Virtual destructor. */
 	~ITargetPlatformManagerModule() = default;
+
+private:
+	/**
+	 * Performs all module setup (SDK validation, platform discovery, format loading), called after the module is 
+	 * loaded and before any other methods are used.
+	 * Separating this from construction allows the async UBT process (kicked off via FDelayedAutoRegisterHelper)
+	 * more time to finish before we block waiting for it.
+	 */
+	virtual void Initialize() = 0;
 };
diff --git a/Engine/Source/Runtime/Core/Private/Misc/CoreMisc.cpp b/Engine/Source/Runtime/Core/Private/Misc/CoreMisc.cpp
--- a/Engine/Source/Runtime/Core/Private/Misc/CoreMisc.cpp
+++ b/Engine/Source/Runtime/Core/Private/Misc/CoreMisc.cpp
@@ -174,7 +174,7 @@ class ITargetPlatformManagerModule* GetTargetPlatformManager(bool bFailOnInitErr
 		{
 			check(IsInGameThread());
 			bInitialized = true;
-			SingletonInterface = FModuleManager::LoadModulePtr<ITargetPlatformManagerModule>("TargetPlatform");
+			SingletonInterface = ITargetPlatformManagerModule::Get();
 
 			FString InitErrors;
 			if (bFailOnInitErrors && SingletonInterface && SingletonInterface->HasInitErrors(&InitErrors))
diff --git a/Engine/Source/Runtime/Core/Private/Misc/DelayedAutoRegister.cpp b/Engine/Source/Runtime/Core/Private/Misc/DelayedAutoRegister.cpp
--- a/Engine/Source/Runtime/Core/Private/Misc/DelayedAutoRegister.cpp
+++ b/Engine/Source/Runtime/Core/Private/Misc/DelayedAutoRegister.cpp
@@ -4,10 +4,9 @@
 
 #include "Containers/Set.h"
 #include "Delegates/Delegate.h"
-#include "Delegates/DelegateBase.h"
-#include "Delegates/DelegateInstancesImpl.h"
-#include "Delegates/IDelegateInstance.h"
+#include "Misc/ConfigCacheIni.h"
 #include "Misc/AssertionMacros.h"
+#include "Modules/ModuleManager.h"
 #include "Templates/Function.h"
 
 DECLARE_MULTICAST_DELEGATE_OneParam(FDelayedAutoRegisterDelegate, EDelayedRegisterRunPhase);
@@ -52,14 +51,74 @@ FDelayedAutoRegisterHelper::FDelayedAutoRegisterHelper(
 	}
 }
 
+namespace
+{
+	const TCHAR* LexToString(EDelayedRegisterRunPhase Phase)
+	{
+		switch (Phase)
+		{
+		case EDelayedRegisterRunPhase::StartOfEnginePreInit: return TEXT("StartOfEnginePreInit");
+		case EDelayedRegisterRunPhase::FileSystemReady: return TEXT("FileSystemReady");
+		case EDelayedRegisterRunPhase::TaskGraphSystemReady: return TEXT("TaskGraphSystemReady");
+		case EDelayedRegisterRunPhase::StatSystemReady: return TEXT("StatSystemReady");
+		case EDelayedRegisterRunPhase::IniSystemReady: return TEXT("IniSystemReady");
+		case EDelayedRegisterRunPhase::EarliestPossiblePluginsLoaded: return TEXT("EarliestPossiblePluginsLoaded");
+		case EDelayedRegisterRunPhase::PreRHIInit: return TEXT("PreRHIInit");
+		case EDelayedRegisterRunPhase::ShaderTypesReady: return TEXT("ShaderTypesReady");
+		case EDelayedRegisterRunPhase::PreObjectSystemReady: return TEXT("PreObjectSystemReady");
+		case EDelayedRegisterRunPhase::ObjectSystemReady: return TEXT("ObjectSystemReady");
+		case EDelayedRegisterRunPhase::DeviceProfileManagerReady: return TEXT("DeviceProfileManagerReady");
+		case EDelayedRegisterRunPhase::EndOfEngineInit: return TEXT("EndOfEngineInit");
+		case EDelayedRegisterRunPhase::LiveCodingReload: return TEXT("LiveCodingReload");
+
+		case EDelayedRegisterRunPhase::NumRunOncePhases:
+		case EDelayedRegisterRunPhase::NumPhases: return TEXT("UnprintableValue");
+
+		default: checkf(false, TEXT("Add Phase to LexToString(EDelayedRegisterRunPhase) [int value: %d]"), (int)Phase); break;
+		}
+
+		return TEXT("Unknown");
+	}
+}
+
 void FDelayedAutoRegisterHelper::RunAndClearDelayedAutoRegisterDelegates(EDelayedRegisterRunPhase RunPhase)
 {
 	const bool bIsRunOncePhase = RunPhase < EDelayedRegisterRunPhase::NumRunOncePhases;
 	if (bIsRunOncePhase)
 	{
-		checkf(!GPhasesAlreadyRun.Contains(RunPhase), TEXT("Delayed Startup phase %d has already run - it is not expected to be run again!"), (int32)RunPhase);
+		checkf(!GPhasesAlreadyRun.Contains(RunPhase), TEXT("Delayed Startup phase %s has already run - it is not expected to be run again!"), LexToString(RunPhase));
 	}
 
+	UE_LOGF(LogInit, Display, "Running DelayedAutoRegister Phase %ls", LexToString(RunPhase));
+
+	// look in .ini for modules to load at this phase (only after ini system is ready!)
+	if (RunPhase >= EDelayedRegisterRunPhase::IniSystemReady)
+	{
+		EDelayedRegisterRunPhase FirstPhase = RunPhase;
+		EDelayedRegisterRunPhase LastPhase = RunPhase;
+
+		// this phase needs to catch up
+		if (RunPhase == EDelayedRegisterRunPhase::IniSystemReady)
+		{
+			FirstPhase = EDelayedRegisterRunPhase::StartOfEnginePreInit;
+		}
+
+		for (EDelayedRegisterRunPhase IniPhase = FirstPhase; IniPhase <= LastPhase; IniPhase = (EDelayedRegisterRunPhase)((int)IniPhase + 1))
+		{
+			TArray<FString> ModuleNames;
+			GConfig->GetArray(TEXT("AutoLoadModulesByPhase"), LexToString(IniPhase), ModuleNames, GEngineIni);
+			for (const FString& ModuleName : ModuleNames)
+			{
+				UE_LOGF(LogInit, Verbose, "Auto-loading module %ls during startup phase %ls", *ModuleName, LexToString(RunPhase));
+				if (FModuleManager::Get().LoadModule(*ModuleName) == nullptr)
+				{
+					UE_LOGF(LogInit, Warning, "Failed to auto-load module %ls during startup phase %ls", *ModuleName, LexToString(RunPhase));
+				}
+			}
+		}
+	}
+
+
 	// run all the delayed functions!
 	GetDelayedAutoRegisterDelegate(RunPhase).Broadcast(RunPhase);
 
