- Changed an early Display log to Log, since it logs before anything can suppress logs (helps with a Verse test that needs specific logging)
#rnx
#Synced-CL 52482450

diff --git a/Engine/Source/Runtime/Core/Private/Misc/DelayedAutoRegister.cpp b/Engine/Source/Runtime/Core/Private/Misc/DelayedAutoRegister.cpp
--- a/Engine/Source/Runtime/Core/Private/Misc/DelayedAutoRegister.cpp
+++ b/Engine/Source/Runtime/Core/Private/Misc/DelayedAutoRegister.cpp
@@ -89,7 +89,7 @@ void FDelayedAutoRegisterHelper::RunAndClearDelayedAutoRegisterDelegates(EDelaye
 		checkf(!GPhasesAlreadyRun.Contains(RunPhase), TEXT("Delayed Startup phase %s has already run - it is not expected to be run again!"), LexToString(RunPhase));
 	}
 
-	UE_LOGF(LogInit, Display, "Running DelayedAutoRegister Phase %ls", LexToString(RunPhase));
+	UE_LOGF(LogInit, Log, "Running DelayedAutoRegister Phase %ls", LexToString(RunPhase));
 
 	// look in .ini for modules to load at this phase (only after ini system is ready!)
 	if (RunPhase >= EDelayedRegisterRunPhase::IniSystemReady)
