[AutoRTFM] Improve performance of transactionally-safe RWLocks by 27%.

The mutex data is now stored in a pool instead of using a TSharedPtr.

#Synced-CL 52411280

diff --git a/Engine/Source/Runtime/Core/Public/Async/TransactionallySafeSharedMutex.h b/Engine/Source/Runtime/Core/Public/Async/TransactionallySafeSharedMutex.h
--- a/Engine/Source/Runtime/Core/Public/Async/TransactionallySafeSharedMutex.h
+++ b/Engine/Source/Runtime/Core/Public/Async/TransactionallySafeSharedMutex.h
@@ -7,6 +7,7 @@
 #include "Templates/SharedPointerFwd.h"
 
 #if UE_AUTORTFM
+#include "Async/TransactionallySafeMutexPool.h"
 #include "AutoRTFM/OpenWrapper.h"
 #include "Templates/SharedPointer.h"
 #endif
@@ -46,16 +47,27 @@ namespace UE
 //                                      effects of the open-threads writes!
 struct FTransactionallySafeSharedMutexDefinition
 {
-	FTransactionallySafeSharedMutexDefinition() : State(MakeShared<FState>())
+	FTransactionallySafeSharedMutexDefinition() : State(UE::Private::TransactionallySafeMutexPool::New<FState>())
 	{
 	}
 
+	~FTransactionallySafeSharedMutexDefinition()
+	{
+		// Normally, State should never be null, but test "Locked Mutex In Destructed Object" expects
+		// that we can memzero a mutex safely. This doesn't crash, but leaks the mutex's State pointer.
+		if (State)
+		{
+			State->ReleaseOnCommit();
+			State = nullptr;
+		}
+	}
+
 	void LockShared()
 	{
 		if (!AutoRTFM::IsIdle())
 		{
 			// Transactionally pessimise LockShared -> Lock.
-			Lock();
+			LockWithinTransaction();
 		}
 		else
 		{
@@ -69,7 +81,7 @@ struct FTransactionallySafeSharedMutexDefinition
 		if (!AutoRTFM::IsIdle())
 		{
 			// Transactionally pessimise UnlockShared -> Unlock.
-			Unlock();
+			UnlockWithinTransaction();
 		}
 		else
 		{
@@ -82,24 +94,7 @@ struct FTransactionallySafeSharedMutexDefinition
 	{
 		if (!AutoRTFM::IsIdle())
 		{
-			UE_AUTORTFM_OPEN
-			{
-				// The transactional system which can increment TransactionalLockCount
-				// is always single-threaded, thus this is safe to check without atomicity.
-				if (0 == State->TransactionalLockCount)
-				{
-					State->Mutex.Lock();
-				}
-
-				State->TransactionalLockCount += 1;
-			};
-
-			// We explicitly copy the state here for the case that `this` was stack
-			// allocated and has already died before the on-abort is hit.
-			UE_AUTORTFM_ONABORT(State = AutoRTFM::TOpenWrapper{this->State})
-			{
-				State.Object->Unlock();
-			};
+			LockWithinTransaction();
 		}
 		else
 		{
@@ -112,12 +107,7 @@ struct FTransactionallySafeSharedMutexDefinition
 	{
 		if (!AutoRTFM::IsIdle())
 		{
-			// We explicitly copy the state here for the case that `this` was stack
-			// allocated and has already died before the on-commit is hit.
-			UE_AUTORTFM_ONCOMMIT(State = AutoRTFM::TOpenWrapper{this->State})
-			{
-				State.Object->Unlock();
-			};
+			UnlockWithinTransaction();
 		}
 		else
 		{
@@ -140,10 +130,55 @@ struct FTransactionallySafeSharedMutexDefinition
 			: TryLockWithinTransaction();
 	}
 
-
 private:
 	UE_NONCOPYABLE(FTransactionallySafeSharedMutexDefinition)
 
+	void LockWithinTransaction()
+	{
+		if (AutoRTFM::IsClosed())
+		{
+			// We explicitly extend the lifetime of `State` here, in case `this` lives on the
+			// stack and ends up being overwritten before the on-abort is hit.
+			State->Acquire();
+
+			UE_AUTORTFM_OPEN
+			{
+				State->LockTransactionalLock();
+			};
+
+			UE_AUTORTFM_ONABORT(State = this->State)
+			{
+				State->UnlockTransactionalLock();
+			};
+
+			State->ReleaseOnCommit();
+		}
+		else
+		{
+			State->LockTransactionalLock();
+		}
+	}
+
+	void UnlockWithinTransaction()
+	{
+		if (AutoRTFM::IsClosed())
+		{
+			// We explicitly extend the lifetime of `State` here, in case `this` lives on the
+			// stack and ends up being overwritten before the on-commit is hit.
+			State->Acquire();
+
+			UE_AUTORTFM_ONCOMMIT(State = this->State)
+			{
+				State->UnlockTransactionalLock();
+				State->ReleaseImmediately();
+			};
+		}
+		else
+		{
+			State->UnlockTransactionalLock();
+		}
+	}
+
 	[[nodiscard]] bool TryLockWithinTransaction()
 	{
 		checkSlow(!AutoRTFM::IsIdle());
@@ -168,12 +203,14 @@ private:
 			// anything with the count or lock.
 			if (LockTaken)
 			{
-				// We explicitly copy the state here for the case that `this` was stack
-				// allocated and has already died before the on-abort is hit.
-				UE_AUTORTFM_ONABORT(State = AutoRTFM::TOpenWrapper{this->State})
+				State->Acquire();
+
+				UE_AUTORTFM_ONABORT(State = this->State)
 				{
-					State.Object->Unlock();
+					State->UnlockTransactionalLock();
 				};
+
+				State->ReleaseOnCommit();
 			}
 		}
 
@@ -184,27 +221,66 @@ private:
 	{
 		::UE::FSharedMutex Mutex;
 		uint32 TransactionalLockCount = 0;
+		uint32 RefCount = 1;
 
 		FState() = default;
 		~FState()
 		{
+			checkSlow(RefCount == 0);
 			ensure(0 == TransactionalLockCount);
 		}
 
-		void Unlock()
+		AUTORTFM_DISABLE void LockTransactionalLock()
 		{
-			ensure(0 != TransactionalLockCount);
+			// The transactional system which can increment TransactionalLockCount
+			// is always single-threaded, thus this is safe to check without atomicity.
+			if (0 == TransactionalLockCount)
+			{
+				Mutex.Lock();
+			}
 
-			TransactionalLockCount -= 1;
+			TransactionalLockCount += 1;
+		}
 
-			if (0 == TransactionalLockCount)
+		AUTORTFM_DISABLE void UnlockTransactionalLock()
+		{
+			ensure(TransactionalLockCount > 0);
+
+			TransactionalLockCount -= 1;
+			if (TransactionalLockCount == 0)
 			{
 				Mutex.Unlock();
 			}
 		}
+
+		AUTORTFM_ENABLE void Acquire()
+		{
+			RefCount += 1;
+		}
+
+		AUTORTFM_ENABLE void ReleaseOnCommit()
+		{
+			checkSlow(RefCount > 0);
+
+			AutoRTFM::OnCommit([this]
+			{
+				ReleaseImmediately();
+			});
+		}
+
+		AUTORTFM_DISABLE void ReleaseImmediately()
+		{
+			checkSlow(RefCount > 0);
+
+			RefCount -= 1;
+			if (RefCount == 0)
+			{
+				UE::Private::TransactionallySafeMutexPool::Delete(this);
+			}
+		}
 	};
 
-	TSharedPtr<FState> State;
+	FState* State = nullptr;
 };
 
 using FTransactionallySafeSharedMutex = ::UE::FTransactionallySafeSharedMutexDefinition;
diff --git a/Engine/Source/Runtime/Core/Public/Misc/TransactionallySafeRWLock.h b/Engine/Source/Runtime/Core/Public/Misc/TransactionallySafeRWLock.h
--- a/Engine/Source/Runtime/Core/Public/Misc/TransactionallySafeRWLock.h
+++ b/Engine/Source/Runtime/Core/Public/Misc/TransactionallySafeRWLock.h
@@ -3,6 +3,7 @@
 #pragma once
 
 #include "AutoRTFM.h"
+#include "HAL/Platform.h"
 
 #if UE_AUTORTFM
 #include "Async/TransactionallySafeSharedMutex.h"
@@ -21,32 +22,32 @@ private:
 	using Super = ::UE::FTransactionallySafeSharedMutex;
 
 public:
-	void ReadLock()
+	UE_REWRITE void ReadLock()
 	{
 		Super::LockShared();
 	}
 
-	void ReadUnlock()
+	UE_REWRITE void ReadUnlock()
 	{
 		Super::UnlockShared();
 	}
 
-	void WriteLock()
+	UE_REWRITE void WriteLock()
 	{
 		Super::Lock();
 	}
 
-	void WriteUnlock()
+	UE_REWRITE void WriteUnlock()
 	{
 		Super::Unlock();
 	}
 
-	bool TryReadLock()
+	UE_REWRITE bool TryReadLock()
 	{
 		return Super::TryLockShared();
 	}
 
-	bool TryWriteLock()
+	UE_REWRITE bool TryWriteLock()
 	{
 		return Super::TryLock();
 	}
