- Changed a Warning to a Log since we unconditionally try to load modules that may not exist in some configs

#rnx
#Synced-CL 52482450

diff --git a/Engine/Source/Runtime/Core/Private/Misc/DelayedAutoRegister.cpp b/Engine/Source/Runtime/Core/Private/Misc/DelayedAutoRegister.cpp
--- a/Engine/Source/Runtime/Core/Private/Misc/DelayedAutoRegister.cpp
+++ b/Engine/Source/Runtime/Core/Private/Misc/DelayedAutoRegister.cpp
@@ -112,7 +112,7 @@ void FDelayedAutoRegisterHelper::RunAndClearDelayedAutoRegisterDelegates(EDelaye
 				UE_LOGF(LogInit, Verbose, "Auto-loading module %ls during startup phase %ls", *ModuleName, LexToString(RunPhase));
 				if (FModuleManager::Get().LoadModule(*ModuleName) == nullptr)
 				{
-					UE_LOGF(LogInit, Warning, "Failed to auto-load module %ls during startup phase %ls", *ModuleName, LexToString(RunPhase));
+					UE_LOGF(LogInit, Log, "Failed to auto-load module %ls during startup phase %ls", *ModuleName, LexToString(RunPhase));
 				}
 			}
 		}
