Cleanup of process and pipe

#Synced-CL 52512740

diff --git a/Engine/Source/Runtime/Core/Public/HAL/PlatformProcess.h b/Engine/Source/Runtime/Core/Public/HAL/PlatformProcess.h
--- a/Engine/Source/Runtime/Core/Public/HAL/PlatformProcess.h
+++ b/Engine/Source/Runtime/Core/Public/HAL/PlatformProcess.h
@@ -18,7 +18,7 @@ namespace UE::HAL
 		{
 		}
 		
-		FPipe(FPipe&& Other)
+		constexpr FPipe(FPipe&& Other)
 			: ReadHandle{std::exchange(Other.ReadHandle, {})}
 			, WriteHandle{std::exchange(Other.WriteHandle, {})}
 		{
@@ -35,19 +35,19 @@ namespace UE::HAL
 			return *this;
 		}
 		
-		FPipe& operator=(FPipe&& Other)
+		constexpr FPipe& operator=(FPipe&& Other)
 		{
 			Swap(ReadHandle, Other.ReadHandle);
 			Swap(WriteHandle, Other.WriteHandle);
 			return *this;
 		}
 		
-		bool operator==(std::nullptr_t) const
+		[[nodiscard]] bool operator==(std::nullptr_t) const
 		{
 			return !ReadHandle;
 		}
 		
-		explicit operator bool() const
+		[[nodiscard]] explicit operator bool() const
 		{
 			return bool(ReadHandle);
 		}
@@ -92,18 +92,18 @@ namespace UE::HAL
 		{
 		}
 
-		void* NativeHandle() const
+		[[nodiscard]] void* NativeHandle() const
 		{
 			return ReadHandle;
 		}
 
-		FString Read()
+		[[nodiscard]] FString Read()
 		{
 			check(*this);
 			return FPlatformProcess::ReadPipe(ReadHandle);
 		}
 		
-		operator FWriteHandle()&
+		[[nodiscard]] operator FWriteHandle()&
 		{
 			return WriteHandle;
 		}
@@ -118,7 +118,7 @@ namespace UE::HAL
 		{
 		}
 		
-		void* NativeHandle() const
+		[[nodiscard]] void* NativeHandle() const
 		{
 			return WriteHandle;
 		}
@@ -129,7 +129,7 @@ namespace UE::HAL
 			return FPlatformProcess::WritePipe(WriteHandle, Message);
 		}
 		
-		operator FReadHandle()&
+		[[nodiscard]] operator FReadHandle()&
 		{
 			return ReadHandle;
 		}
@@ -165,10 +165,9 @@ namespace UE::HAL
 		{
 		}
 		
-		FProcess(FProcess&& Other)
-			: Handle{Other.Handle}
+		constexpr FProcess(FProcess&& Other)
+			: Handle{std::exchange(Other.Handle, {})}
 		{
-			Other.Handle.Reset();
 		}
 		
 		~FProcess()
@@ -182,24 +181,24 @@ namespace UE::HAL
 			return *this;
 		}
 		
-		FProcess& operator=(FProcess&& Other)
+		constexpr FProcess& operator=(FProcess&& Other)
 		{
 			Swap(Handle, Other.Handle);
 			return *this;
 		}
 		
-		decltype(auto) NativeHandle() const
+		[[nodiscard]] decltype(auto) NativeHandle() const
 		{
 			return Handle.Get();
 		}
 		
-		bool IsRunning()
+		[[nodiscard]] bool IsRunning()
 		{
 			check(*this);
 			return FPlatformProcess::IsProcRunning(Handle);
 		}
 		
-		TOptional<int32> GetExitCode()
+		[[nodiscard]] TOptional<int32> GetExitCode()
 		{
 			check(*this);
 			int32 ExitCode;
@@ -224,12 +223,12 @@ namespace UE::HAL
 			return *this;
 		}
 
-		bool operator==(std::nullptr_t) const
+		[[nodiscard]] bool operator==(std::nullptr_t) const
 		{
 			return !Handle.IsValid();
 		}
 
-		explicit operator bool() const
+		[[nodiscard]] explicit operator bool() const
 		{
 			return Handle.IsValid();
 		}
