Unix Platform File: Fix unspecific mkdir-p failure log

#rnx
#Synced-CL 52515964

diff --git a/Engine/Source/Runtime/Core/Private/Unix/UnixPlatformFile.cpp b/Engine/Source/Runtime/Core/Private/Unix/UnixPlatformFile.cpp
--- a/Engine/Source/Runtime/Core/Private/Unix/UnixPlatformFile.cpp
+++ b/Engine/Source/Runtime/Core/Private/Unix/UnixPlatformFile.cpp
@@ -62,6 +62,17 @@ namespace
 			bLoggingError = false; \
 		} \
 	}
+
+	// Wrapper around UE_LOG to avoid infinite logging if the act of logging causes an error
+	#define UE_LOGF_UNIX_FILE(Verbosity, Format, ...) \
+	{ \
+		if (!bLoggingError) \
+		{ \
+			bLoggingError = true;  \
+			UE_LOGF(LogUnixPlatformFile, Verbosity, Format, ##__VA_ARGS__); \
+			bLoggingError = false; \
+		} \
+	}
 }
 
 /**
@@ -783,6 +794,7 @@ public:
 
 		// try opening right away
 		int32 Handle = open(TCHAR_TO_UTF8(*Filename), O_RDONLY | O_CLOEXEC);
+		const int ErrNo = errno;
 		if (Handle != -1)
 		{
 			MappedToFilename = Filename;
@@ -790,9 +802,8 @@ public:
 		else
 		{
 			// log non-standard errors only
-			if (ENOENT != errno)
+			if (ENOENT != ErrNo)
 			{
-				int ErrNo = errno;
 				UE_LOG(LogUnixPlatformFile, Warning, TEXT( "open('%s', O_RDONLY | O_CLOEXEC) failed: errno=%d (%s)" ), *Filename, ErrNo, UTF8_TO_TCHAR(strerror(ErrNo)));
 			}
 			else
@@ -1620,7 +1631,7 @@ bool FUnixPlatformFile::CreateDirectoriesFromPath(const TCHAR* Path)
 					continue;
 				}
 
-				UE_LOG_UNIX_FILE(Warning, TEXT( "create dir('%s') failed: errno=%d (%s)" ), UTF8_TO_TCHAR(DirPath), ErrNo, UTF8_TO_TCHAR(strerror(ErrNo)));
+				UE_LOGF_UNIX_FILE(Warning, "create dir('%s') failed: errno=%d (%s)", SubPath, ErrNo, strerror(ErrNo));
 
 				FMemory::Free(DirPath);
 				FMemory::Free(SubPath);
