[AutoRTFM] Improve performance of TransactionallySafeSharedRecursiveMutex by 31%.

The mutex data is now stored in a pool instead of using a TSharedPtr. This class is functionally identical to transactionally-safe RWLocks, so the performance uplift is also similar.

#Synced-CL 52411280

diff --git a/Engine/Source/Runtime/Core/Public/Async/TransactionallySafeSharedRecursiveMutex.h b/Engine/Source/Runtime/Core/Public/Async/TransactionallySafeSharedRecursiveMutex.h
--- a/Engine/Source/Runtime/Core/Public/Async/TransactionallySafeSharedRecursiveMutex.h
+++ b/Engine/Source/Runtime/Core/Public/Async/TransactionallySafeSharedRecursiveMutex.h
@@ -46,16 +46,27 @@ namespace UE
 //                                      effects of the open-threads writes!
 struct FTransactionallySafeSharedRecursiveMutexDefinition
 {
-	FTransactionallySafeSharedRecursiveMutexDefinition() : State(MakeShared<FState>())
+	FTransactionallySafeSharedRecursiveMutexDefinition() : State(UE::Private::TransactionallySafeMutexPool::New<FState>())
 	{
 	}
 
+	~FTransactionallySafeSharedRecursiveMutexDefinition()
+	{
+		// Normally, State should never be null, but test "Locked Mutex In Destructed Object" expects
+		// that we can memzero a mutex safely. This doesn't crash, but leaks the mutex's State pointer.
+		if (State)
+		{
+			State->ReleaseOnCommit();
+			State = nullptr;
+		}
+	}
+
 	void LockShared(Core::Private::FSharedRecursiveMutexLink& Link)
 	{
 		if (!AutoRTFM::IsIdle())
 		{
 			// Transactionally pessimise LockShared -> Lock.
-			Lock();
+			LockWithinTransaction();
 		}
 		else
 		{
@@ -69,7 +80,7 @@ struct FTransactionallySafeSharedRecursiveMutexDefinition
 		if (!AutoRTFM::IsIdle())
 		{
 			// Transactionally pessimise UnlockShared -> Unlock.
-			Unlock();
+			UnlockWithinTransaction();
 		}
 		else
 		{
@@ -82,31 +93,7 @@ struct FTransactionallySafeSharedRecursiveMutexDefinition
 	{
 		if (!AutoRTFM::IsIdle())
 		{
-			UE_AUTORTFM_OPEN
-			{
-				// The transactional system which can increment TransactionalLockCount
-				// is always single-threaded, thus this is safe to check without atomicity.
-				if (0 == State->TransactionalLockCount)
-				{
-					State->Mutex.Lock();
-				}
-
-				State->TransactionalLockCount += 1;
-			};
-
-			// We explicitly copy the state here for the case that `this` was stack
-			// allocated and has already died before the on-abort is hit.
-			AutoRTFM::OnAbort([State = AutoRTFM::TOpenWrapper{this->State}]
-			{
-				ensure(0 != State.Object->TransactionalLockCount);
-
-				State.Object->TransactionalLockCount -= 1;
-
-				if (0 == State.Object->TransactionalLockCount)
-				{
-					State.Object->Mutex.Unlock();
-				}
-			});
+			LockWithinTransaction();
 		}
 		else
 		{
@@ -119,19 +106,7 @@ struct FTransactionallySafeSharedRecursiveMutexDefinition
 	{
 		if (!AutoRTFM::IsIdle())
 		{
-			// We explicitly copy the state here for the case that `this` was stack
-			// allocated and has already died before the on-commit is hit.
-			AutoRTFM::OnCommit([State = AutoRTFM::TOpenWrapper{this->State}]
-			{
-				ensure(0 != State.Object->TransactionalLockCount);
-
-				State.Object->TransactionalLockCount -= 1;
-
-				if (0 == State.Object->TransactionalLockCount)
-				{
-					State.Object->Mutex.Unlock();
-				}
-			});
+			UnlockWithinTransaction();
 		}
 		else
 		{
@@ -173,20 +148,116 @@ struct FTransactionallySafeSharedRecursiveMutexDefinition
 private:
 	UE_NONCOPYABLE(FTransactionallySafeSharedRecursiveMutexDefinition)
 
+	void LockWithinTransaction()
+	{
+		if (AutoRTFM::IsClosed())
+		{
+			// We explicitly extend the lifetime of `State` here, in case `this` lives on the
+			// stack and ends up being overwritten before the on-abort is hit.
+			State->Acquire();
+
+			UE_AUTORTFM_OPEN
+			{
+				State->LockTransactionalLock();
+			};
+
+			UE_AUTORTFM_ONABORT(State = this->State)
+			{
+				State->UnlockTransactionalLock();
+			};
+
+			State->ReleaseOnCommit();
+		}
+		else
+		{
+			State->LockTransactionalLock();
+		}
+	}
+
+	void UnlockWithinTransaction()
+	{
+		if (AutoRTFM::IsClosed())
+		{
+			// We explicitly extend the lifetime of `State` here, in case `this` lives on the
+			// stack and ends up being overwritten before the on-commit is hit.
+			State->Acquire();
+
+			UE_AUTORTFM_ONCOMMIT(State = this->State)
+			{
+				State->UnlockTransactionalLock();
+				State->ReleaseImmediately();
+			};
+		}
+		else
+		{
+			State->UnlockTransactionalLock();
+		}
+	}
+
 	struct FState final
 	{
 		::UE::FSharedRecursiveMutex Mutex;
 		uint32 TransactionalLockCount = 0;
+		uint32 RefCount = 1;
 
 		FState() = default;
-		
 		~FState()
 		{
+			checkSlow(RefCount == 0);
 			ensure(0 == TransactionalLockCount);
 		}
+
+		AUTORTFM_DISABLE void LockTransactionalLock()
+		{
+			// The transactional system which can increment TransactionalLockCount
+			// is always single-threaded, thus this is safe to check without atomicity.
+			if (0 == TransactionalLockCount)
+			{
+				Mutex.Lock();
+			}
+
+			TransactionalLockCount += 1;
+		}
+
+		AUTORTFM_DISABLE void UnlockTransactionalLock()
+		{
+			ensure(TransactionalLockCount > 0);
+
+			TransactionalLockCount -= 1;
+			if (TransactionalLockCount == 0)
+			{
+				Mutex.Unlock();
+			}
+		}
+
+		AUTORTFM_ENABLE void Acquire()
+		{
+			RefCount += 1;
+		}
+
+		AUTORTFM_ENABLE void ReleaseOnCommit()
+		{
+			checkSlow(RefCount > 0);
+
+			AutoRTFM::OnCommit([this]
+			{
+				ReleaseImmediately();
+			});
+		}
+
+		AUTORTFM_DISABLE void ReleaseImmediately()
+		{
+			checkSlow(RefCount > 0);
+
+			RefCount -= 1;
+			if (RefCount == 0)
+			{
+				UE::Private::TransactionallySafeMutexPool::Delete(this);
+			}
+		}
 	};
 
-	TSharedPtr<FState> State;
+	FState* State = nullptr;
 };
 
 using FTransactionallySafeSharedRecursiveMutex = ::UE::FTransactionallySafeSharedRecursiveMutexDefinition;
