Fix common typo being flagged by our pre-submit checks.

#rnx

- Infomation -> Information.

#Synced-CL 52535685

diff --git a/Engine/Source/Runtime/Core/Private/Windows/WindowsPlatformProcess.cpp b/Engine/Source/Runtime/Core/Private/Windows/WindowsPlatformProcess.cpp
--- a/Engine/Source/Runtime/Core/Private/Windows/WindowsPlatformProcess.cpp
+++ b/Engine/Source/Runtime/Core/Private/Windows/WindowsPlatformProcess.cpp
@@ -694,22 +694,22 @@ FProcHandle FWindowsPlatformProcess::CreateProcInternal(const TCHAR* URL, const
 		},
 		AttributeList
 	};
-	PROCESS_INFORMATION ProcessInfomation;
+	PROCESS_INFORMATION ProcessInformation;
 
-	if (!CreateProcess(nullptr, CommandLine.GetCharArray().GetData(), nullptr, nullptr, bInheritHandles || AttributeList, CreationFlags, nullptr, OptionalWorkingDirectory, &StartupInfoEx.StartupInfo, &ProcessInfomation))
+	if (!CreateProcess(nullptr, CommandLine.GetCharArray().GetData(), nullptr, nullptr, bInheritHandles || AttributeList, CreationFlags, nullptr, OptionalWorkingDirectory, &StartupInfoEx.StartupInfo, &ProcessInformation))
 	{
 		HandleError();
 		return {};
 	}
 
-	CloseHandle(ProcessInfomation.hThread);
+	CloseHandle(ProcessInformation.hThread);
 
 	if (OutProcessID)
 	{
-		*OutProcessID = ProcessInfomation.dwProcessId;
+		*OutProcessID = ProcessInformation.dwProcessId;
 	}
 
-	return FProcHandle{ProcessInfomation.hProcess};
+	return FProcHandle{ProcessInformation.hProcess};
 }
 #include "Windows/HideWindowsPlatformTypes.h"
 
