[InheritedContext]
- Add per-thread deduplication of extensions

#Synced-CL 52534307

diff --git a/Engine/Source/Runtime/Core/Private/Async/InheritedContext.cpp b/Engine/Source/Runtime/Core/Private/Async/InheritedContext.cpp
--- a/Engine/Source/Runtime/Core/Private/Async/InheritedContext.cpp
+++ b/Engine/Source/Runtime/Core/Private/Async/InheritedContext.cpp
@@ -146,15 +146,42 @@ struct FInheritedContextExtensionData
 
 FInheritedContextExtensionScope::FInheritedContextExtensionScope(const FInheritedContextExtension& InExtension)
 	: Extension(InExtension)
-	, Next(GExtensionScopeHead)
+	, Next(nullptr)
+	, bPushed(true)
 {
+	// O(1) dedup via TLS: if this extension is already active on this thread, skip pushing.
+	// This happens with nested scopes sharing a single static extension (e.g. FScopeTimer).
+	const uint32 TlsSlot = Extension.Impl->TlsSlot;
+	if (FPlatformTLS::IsValidTlsSlot(TlsSlot) && FPlatformTLS::GetTlsValue(TlsSlot) != nullptr)
+	{
+		bPushed = false;
+		return;
+	}
+
+	Next = GExtensionScopeHead;
 	GExtensionScopeHead = this;
+
+	if (FPlatformTLS::IsValidTlsSlot(TlsSlot))
+	{
+		FPlatformTLS::SetTlsValue(TlsSlot, this);
+	}
 }
 
 FInheritedContextExtensionScope::~FInheritedContextExtensionScope()
 {
+	if (!bPushed)
+	{
+		return;
+	}
+
 	checkf(GExtensionScopeHead == this, TEXT("FInheritedContextExtensionScope destroyed out of LIFO order"));
 	GExtensionScopeHead = Next;
+
+	const uint32 TlsSlot = Extension.Impl->TlsSlot;
+	if (FPlatformTLS::IsValidTlsSlot(TlsSlot))
+	{
+		FPlatformTLS::SetTlsValue(TlsSlot, nullptr);
+	}
 }
 
 // ---- FInheritedContextScope ----
diff --git a/Engine/Source/Runtime/Core/Public/Async/InheritedContext.h b/Engine/Source/Runtime/Core/Public/Async/InheritedContext.h
--- a/Engine/Source/Runtime/Core/Public/Async/InheritedContext.h
+++ b/Engine/Source/Runtime/Core/Public/Async/InheritedContext.h
@@ -73,6 +73,15 @@ namespace UE
 			uint32        DataSize;
 			uint32        DataAlign;
 			void*         UserData;   // Typically pointer to the variable being propagated
+			uint32        TlsSlot = FPlatformTLS::InvalidTlsSlot; // Per-thread active flag for O(1) dedup in scope chain
+
+			~FInterface()
+			{
+				if (FPlatformTLS::IsValidTlsSlot(TlsSlot))
+				{
+					FPlatformTLS::FreeTlsSlot(TlsSlot);
+				}
+			}
 		};
 
 		TSharedPtr<FInterface> Impl;
@@ -129,6 +138,7 @@ namespace UE
 		Ext.Impl->DataSize  = sizeof(T);
 		Ext.Impl->DataAlign = alignof(T);
 		Ext.Impl->UserData  = reinterpret_cast<void*>(FnPtr);
+		Ext.Impl->TlsSlot  = FPlatformTLS::AllocTlsSlot();
 		return Ext;
 	}
 
@@ -167,6 +177,7 @@ namespace UE
 		Ext.Impl->DataSize  = sizeof(T);
 		Ext.Impl->DataAlign = alignof(T);
 		Ext.Impl->UserData  = nullptr; // Getter/Setter are baked into template callbacks; UserData is not needed
+		Ext.Impl->TlsSlot  = FPlatformTLS::AllocTlsSlot();
 		return Ext;
 	}
 
@@ -261,6 +272,7 @@ namespace UE
 
 		FInheritedContextExtension Extension;
 		FInheritedContextExtensionScope* Next;
+		bool bPushed;
 	};
 
 	// Restores an inherited context for the current scope.
