[RaceDetector]
- Add support for race detection under autortfm

#Synced-CL 52534307

diff --git a/Engine/Source/Programs/UnrealBuildTool/Configuration/UEBuildModule.cs b/Engine/Source/Programs/UnrealBuildTool/Configuration/UEBuildModule.cs
--- a/Engine/Source/Programs/UnrealBuildTool/Configuration/UEBuildModule.cs
+++ b/Engine/Source/Programs/UnrealBuildTool/Configuration/UEBuildModule.cs
@@ -760,7 +760,8 @@ namespace UnrealBuildTool
 				{
 					ApiDefinitionIndex = 2;
 				}
-				else if (!(!Binary.bAllowExports || Rules.Target.bMergeModules))
+				else if (!(!Binary.bAllowExports || Rules.Target.bMergeModules)
+					|| (Rules.Target.bMergeModules && Rules.ModuleSymbolVisibility == ModuleRules.SymbolVisibility.VisibileForDll))
 				{
 					ApiDefinitionIndex = 1;
 				}
diff --git a/Engine/Source/Programs/UnrealBuildTool/Configuration/UEBuildTarget.cs b/Engine/Source/Programs/UnrealBuildTool/Configuration/UEBuildTarget.cs
--- a/Engine/Source/Programs/UnrealBuildTool/Configuration/UEBuildTarget.cs
+++ b/Engine/Source/Programs/UnrealBuildTool/Configuration/UEBuildTarget.cs
@@ -4432,6 +4432,12 @@ namespace UnrealBuildTool
 			List<UEBuildModuleCPP> engineModules = [];
 			HashSet<UEBuildModuleCPP> excludedModules = [];
 
+			// Exclude the Instrumentation module from merging so it stays as its own small DLL.
+			// Its symbols are dllimport'd by PerModuleInline.inl in every other DLL, so it needs
+			// a standalone import library. Combined with ModuleSymbolVisibility = VisibileForDll
+			// in Instrumentation.Build.cs to ensure DLLEXPORT even with bMergeModules.
+			allModules.RemoveWhere(m => m.Name == "Instrumentation");
+
 			// Collect engine modules and find the ones that are explicitly moved to common
 			foreach (UEBuildModuleCPP m in allModules)
 			{
@@ -6771,7 +6777,8 @@ namespace UnrealBuildTool
 			}
 
 			// If compiling with the AutoRTFM compiler, implicitly depend on the AutoRTFM runtime.
-			if (Rules.bUseAutoRTFMCompiler && ModuleName != "AutoRTFM")
+			// Exclude Instrumentation to avoid circular dependency (Instrumentation -> AutoRTFM -> Instrumentation).
+			if (Rules.bUseAutoRTFMCompiler && ModuleName != "AutoRTFM" && ModuleName != "Instrumentation")
 			{
 				RulesObject.PrivateDependencyModuleNames.Add("AutoRTFM");
 			}
diff --git a/Engine/Source/Programs/UnrealBuildTool/Platform/Windows/UEBuildWindows.cs b/Engine/Source/Programs/UnrealBuildTool/Platform/Windows/UEBuildWindows.cs
--- a/Engine/Source/Programs/UnrealBuildTool/Platform/Windows/UEBuildWindows.cs
+++ b/Engine/Source/Programs/UnrealBuildTool/Platform/Windows/UEBuildWindows.cs
@@ -1331,7 +1331,7 @@ namespace UnrealBuildTool
 				Target.WindowsPlatform.Compiler = WindowsCompiler.Default;
 			}
 
-			if (Target.WindowsPlatform.bEnableInstrumentation)
+			if (Target.WindowsPlatform.bEnableInstrumentation && !Target.bUseAutoRTFMCompiler)
 			{
 				Target.WindowsPlatform.Compiler = WindowsCompiler.ClangInstrument;
 
diff --git a/Engine/Source/Programs/UnrealBuildTool/Platform/Windows/VCToolChain.cs b/Engine/Source/Programs/UnrealBuildTool/Platform/Windows/VCToolChain.cs
--- a/Engine/Source/Programs/UnrealBuildTool/Platform/Windows/VCToolChain.cs
+++ b/Engine/Source/Programs/UnrealBuildTool/Platform/Windows/VCToolChain.cs
@@ -811,6 +811,13 @@ namespace UnrealBuildTool
 			if (Target.WindowsPlatform.bEnableInstrumentation)
 			{
 				AddDefinition(Arguments, "USING_INSTRUMENTATION=1");
+
+				// When AutoRTFM is also active, verse-clang-cl handles instrumentation
+				// via -fepictsan instead of using instr-clang-cl.exe.
+				if (Target.WindowsPlatform.Compiler == WindowsCompiler.ClangRTFM)
+				{
+					Arguments.Add("-fepictsan");
+				}
 			}
 
 			// Address sanitizer
diff --git a/Engine/Source/Runtime/Core/Internal/Instrumentation/Containers.h b/Engine/Source/Runtime/Core/Internal/Instrumentation/Containers.h
--- a/Engine/Source/Runtime/Core/Internal/Instrumentation/Containers.h
+++ b/Engine/Source/Runtime/Core/Internal/Instrumentation/Containers.h
@@ -17,11 +17,11 @@
 #if PLATFORM_WINDOWS
 class FInstrumentationSafeWinAllocator {
 public:
-	static INSTRUMENTATION_FUNCTION_ATTRIBUTES void* Alloc(SIZE_T Size);
-	static INSTRUMENTATION_FUNCTION_ATTRIBUTES void Free(void* Data);
-	static INSTRUMENTATION_FUNCTION_ATTRIBUTES void* AllocWithGuards(SIZE_T Size);
-	static INSTRUMENTATION_FUNCTION_ATTRIBUTES void FreeWithGuards(void* Data);
-	static INSTRUMENTATION_FUNCTION_ATTRIBUTES void* Realloc(void* Data, SIZE_T Size, SIZE_T PreviousSize);
+	INSTRUMENTATION_FUNCTION_ATTRIBUTES static void* Alloc(SIZE_T Size);
+	INSTRUMENTATION_FUNCTION_ATTRIBUTES static void Free(void* Data);
+	INSTRUMENTATION_FUNCTION_ATTRIBUTES static void* AllocWithGuards(SIZE_T Size);
+	INSTRUMENTATION_FUNCTION_ATTRIBUTES static void FreeWithGuards(void* Data);
+	INSTRUMENTATION_FUNCTION_ATTRIBUTES static void* Realloc(void* Data, SIZE_T Size, SIZE_T PreviousSize);
 };
 
 using TInstrumentationSafeBaseAllocator = FInstrumentationSafeWinAllocator;
@@ -90,7 +90,7 @@ public:
 		 * Assumes that the allocator is currently empty, i.e. memory may be allocated but any existing elements have already been destructed (if necessary).
 		 * @param Other - The allocator to move the state from.  This allocator should be left in a valid empty state.
 		 */
-		FORCEINLINE INSTRUMENTATION_FUNCTION_ATTRIBUTES void MoveToEmpty(ForAnyElementType& Other)
+		INSTRUMENTATION_FUNCTION_ATTRIBUTES FORCEINLINE void MoveToEmpty(ForAnyElementType& Other)
 		{
 			if (Data)
 			{
@@ -102,7 +102,7 @@ public:
 		}
 
 		/** Destructor. */
-		FORCEINLINE INSTRUMENTATION_FUNCTION_ATTRIBUTES ~ForAnyElementType()
+		INSTRUMENTATION_FUNCTION_ATTRIBUTES FORCEINLINE ~ForAnyElementType()
 		{
 			if (Data)
 			{
@@ -111,7 +111,7 @@ public:
 		}
 
 		// FContainerAllocatorInterface
-		FORCEINLINE INSTRUMENTATION_FUNCTION_ATTRIBUTES void* GetAllocation() const
+		INSTRUMENTATION_FUNCTION_ATTRIBUTES FORCEINLINE void* GetAllocation() const
 		{
 			return Data;
 		}
@@ -141,27 +141,27 @@ public:
 				}
 			}
 		}
-		FORCEINLINE INSTRUMENTATION_FUNCTION_ATTRIBUTES SizeType CalculateSlackReserve(SizeType NumElements, SIZE_T NumBytesPerElement) const
+		INSTRUMENTATION_FUNCTION_ATTRIBUTES FORCEINLINE SizeType CalculateSlackReserve(SizeType NumElements, SIZE_T NumBytesPerElement) const
 		{
 			return DefaultCalculateSlackReserve(NumElements, NumBytesPerElement, true);
 		}
-		FORCEINLINE INSTRUMENTATION_FUNCTION_ATTRIBUTES SizeType CalculateSlackReserve(SizeType NumElements, SIZE_T NumBytesPerElement, uint32 AlignmentOfElement) const
+		INSTRUMENTATION_FUNCTION_ATTRIBUTES FORCEINLINE SizeType CalculateSlackReserve(SizeType NumElements, SIZE_T NumBytesPerElement, uint32 AlignmentOfElement) const
 		{
 			return DefaultCalculateSlackReserve(NumElements, NumBytesPerElement, true, (uint32)AlignmentOfElement);
 		}
-		FORCEINLINE INSTRUMENTATION_FUNCTION_ATTRIBUTES SizeType CalculateSlackShrink(SizeType NumElements, SizeType NumAllocatedElements, SIZE_T NumBytesPerElement) const
+		INSTRUMENTATION_FUNCTION_ATTRIBUTES FORCEINLINE SizeType CalculateSlackShrink(SizeType NumElements, SizeType NumAllocatedElements, SIZE_T NumBytesPerElement) const
 		{
 			return DefaultCalculateSlackShrink(NumElements, NumAllocatedElements, NumBytesPerElement, true);
 		}
-		FORCEINLINE INSTRUMENTATION_FUNCTION_ATTRIBUTES SizeType CalculateSlackShrink(SizeType NumElements, SizeType NumAllocatedElements, SIZE_T NumBytesPerElement, uint32 AlignmentOfElement) const
+		INSTRUMENTATION_FUNCTION_ATTRIBUTES FORCEINLINE SizeType CalculateSlackShrink(SizeType NumElements, SizeType NumAllocatedElements, SIZE_T NumBytesPerElement, uint32 AlignmentOfElement) const
 		{
 			return DefaultCalculateSlackShrink(NumElements, NumAllocatedElements, NumBytesPerElement, true, (uint32)AlignmentOfElement);
 		}
-		FORCEINLINE INSTRUMENTATION_FUNCTION_ATTRIBUTES SizeType CalculateSlackGrow(SizeType NumElements, SizeType NumAllocatedElements, SIZE_T NumBytesPerElement) const
+		INSTRUMENTATION_FUNCTION_ATTRIBUTES FORCEINLINE SizeType CalculateSlackGrow(SizeType NumElements, SizeType NumAllocatedElements, SIZE_T NumBytesPerElement) const
 		{
 			return DefaultCalculateSlackGrow(NumElements, NumAllocatedElements, NumBytesPerElement, true);
 		}
-		FORCEINLINE INSTRUMENTATION_FUNCTION_ATTRIBUTES SizeType CalculateSlackGrow(SizeType NumElements, SizeType NumAllocatedElements, SIZE_T NumBytesPerElement, uint32 AlignmentOfElement) const
+		INSTRUMENTATION_FUNCTION_ATTRIBUTES FORCEINLINE SizeType CalculateSlackGrow(SizeType NumElements, SizeType NumAllocatedElements, SIZE_T NumBytesPerElement, uint32 AlignmentOfElement) const
 		{
 			return DefaultCalculateSlackGrow(NumElements, NumAllocatedElements, NumBytesPerElement, true, (uint32)AlignmentOfElement);
 		}
@@ -198,7 +198,7 @@ public:
 		{
 		}
 
-		FORCEINLINE INSTRUMENTATION_FUNCTION_ATTRIBUTES ElementType* GetAllocation() const
+		INSTRUMENTATION_FUNCTION_ATTRIBUTES FORCEINLINE ElementType* GetAllocation() const
 		{
 			return (ElementType*)ForAnyElementType::GetAllocation();
 		}
@@ -222,7 +222,7 @@ class TInstrumentationSafeSparseSetAllocator
 public:
 
 	/** Computes the number of hash buckets to use for a given number of elements. */
-	static FORCEINLINE INSTRUMENTATION_FUNCTION_ATTRIBUTES uint32 GetNumberOfHashBuckets(uint32 NumHashedElements)
+	INSTRUMENTATION_FUNCTION_ATTRIBUTES static FORCEINLINE uint32 GetNumberOfHashBuckets(uint32 NumHashedElements)
 	{
 		if (NumHashedElements >= MinNumberOfHashedElements)
 		{
diff --git a/Engine/Source/Runtime/Core/Internal/Sanitizer/RaceDetectorInterface.h b/Engine/Source/Runtime/Core/Internal/Sanitizer/RaceDetectorInterface.h
--- a/Engine/Source/Runtime/Core/Internal/Sanitizer/RaceDetectorInterface.h
+++ b/Engine/Source/Runtime/Core/Internal/Sanitizer/RaceDetectorInterface.h
@@ -7,6 +7,7 @@
 #if USING_INSTRUMENTATION
 
 #include "CoreTypes.h"
+#include "Instrumentation/Defines.h"
 #include "Sanitizer/RaceDetectorTypes.h"
 
 DECLARE_LOG_CATEGORY_EXTERN(LogRaceDetector, Log, All);
@@ -14,18 +15,18 @@ DECLARE_LOG_CATEGORY_EXTERN(LogRaceDetector, Log, All);
 namespace UE::Sanitizer::RaceDetector {
 
 	// Returns whether we should instrument depending on the current context state.
-	bool ShouldInstrument(FContext& Context);
+	INSTRUMENTATION_FUNCTION_ATTRIBUTES bool ShouldInstrument(FContext& Context);
 	// Gets the current thread context, could be nullptr.
-	FContext* GetThreadContext();
+	INSTRUMENTATION_FUNCTION_ATTRIBUTES FContext* GetThreadContext();
 	// Hints the sanitizer that this memory range is being freed.
-	void FreeMemoryRange(void* Ptr, uint64 Size);
+	INSTRUMENTATION_FUNCTION_ATTRIBUTES void FreeMemoryRange(void* Ptr, uint64 Size);
 
 	// Makes sure the current thread has a context and returns it.
-	FContext& EnsureCurrentContext();
+	INSTRUMENTATION_FUNCTION_ATTRIBUTES FContext& EnsureCurrentContext();
 	// Releases the current thread context.
-	void ReleaseCurrentContext();
+	INSTRUMENTATION_FUNCTION_ATTRIBUTES void ReleaseCurrentContext();
 	// Returns a sync object for the given address, initialize one if there isn't one already.
-	FSyncObjectRef GetSyncObject(FContext& Context, void* SyncAddr);
+	INSTRUMENTATION_FUNCTION_ATTRIBUTES FSyncObjectRef GetSyncObject(FContext& Context, void* SyncAddr);
 }
 
 #endif // USING_INSTRUMENTATION
\ No newline at end of file
diff --git a/Engine/Source/Runtime/Core/Internal/Sanitizer/RaceDetectorPlatform.h b/Engine/Source/Runtime/Core/Internal/Sanitizer/RaceDetectorPlatform.h
--- a/Engine/Source/Runtime/Core/Internal/Sanitizer/RaceDetectorPlatform.h
+++ b/Engine/Source/Runtime/Core/Internal/Sanitizer/RaceDetectorPlatform.h
@@ -7,90 +7,91 @@
 #if USING_INSTRUMENTATION
 
 #include "CoreTypes.h"
+#include "Instrumentation/Defines.h"
 
 namespace UE::Sanitizer::RaceDetector {
 	namespace Platform {
 		// Performs any platform specific initialization.
-		bool  InitializePlatform();
+		INSTRUMENTATION_FUNCTION_ATTRIBUTES bool  InitializePlatform();
 		// Performs any platform specific cleanup.
-		bool  CleanupPlatform();
+		INSTRUMENTATION_FUNCTION_ATTRIBUTES bool  CleanupPlatform();
 		// Prepare address space of shadow memory
-		void  InitShadowMemory();
+		INSTRUMENTATION_FUNCTION_ATTRIBUTES void  InitShadowMemory();
 		// Returns the command line the process was started with.
-		const TCHAR* GetCommandLine();
+		INSTRUMENTATION_FUNCTION_ATTRIBUTES const TCHAR* GetCommandLine();
 		// Returns the page size granularity.
-		UPTRINT GetPageSize();
-		// Returns the base address of the shadow memory address space. 
+		INSTRUMENTATION_FUNCTION_ATTRIBUTES UPTRINT GetPageSize();
+		// Returns the base address of the shadow memory address space.
 		// Note: InitShadowMemory must have been called before, otherwise this is meaningless.
-		UPTRINT GetShadowMemoryBase();
+		INSTRUMENTATION_FUNCTION_ATTRIBUTES UPTRINT GetShadowMemoryBase();
 		// Returns the size of the entire shadow memory address space.
 		// Note: InitShadowMemory must have been called before, otherwise this is meaningless.
-		UPTRINT GetShadowMemorySize();
+		INSTRUMENTATION_FUNCTION_ATTRIBUTES UPTRINT GetShadowMemorySize();
 		// Get the base of the shadow memory for clock banks.
-		UPTRINT GetShadowClockBase();
+		INSTRUMENTATION_FUNCTION_ATTRIBUTES UPTRINT GetShadowClockBase();
 		// Returns the number of bytes currently mapped in the shadow memory.
-		uint64 GetShadowMemoryUsage();
+		INSTRUMENTATION_FUNCTION_ATTRIBUTES uint64 GetShadowMemoryUsage();
 		// Returns whether a particular range is already accessible in shadow memory.
-		bool IsShadowMemoryMapped(UPTRINT Base, UPTRINT Size);
+		INSTRUMENTATION_FUNCTION_ATTRIBUTES bool IsShadowMemoryMapped(UPTRINT Base, UPTRINT Size);
 		// Maps a range in shadow memory so that it is safe to access.
-		void MapShadowMemory(UPTRINT Base, UPTRINT Size);
+		INSTRUMENTATION_FUNCTION_ATTRIBUTES void MapShadowMemory(UPTRINT Base, UPTRINT Size);
 		// Unmaps the entire range of shadow memory.
-		void UnmapShadowMemory();
+		INSTRUMENTATION_FUNCTION_ATTRIBUTES void UnmapShadowMemory();
 		// Returns whether a debugger is currently attached to our process.
-		bool IsDebuggerPresent();
+		INSTRUMENTATION_FUNCTION_ATTRIBUTES bool IsDebuggerPresent();
 		// Returns whether there is any page currently mapped in shadow memory.
-		bool HasShadowMemoryMapped();
+		INSTRUMENTATION_FUNCTION_ATTRIBUTES bool HasShadowMemoryMapped();
 		// Sends a hint to visual studio to hide first chance exception to reduce noise caused by shadow memory access.
-		void HideFirstChanceExceptionInVisualStudio();
+		INSTRUMENTATION_FUNCTION_ATTRIBUTES void HideFirstChanceExceptionInVisualStudio();
 		// Sleeps for the given amount of milliseconds.
-		void SleepMS(uint32 Milliseconds);
+		INSTRUMENTATION_FUNCTION_ATTRIBUTES void SleepMS(uint32 Milliseconds);
 		// Capture the current callstack.
-		uint16 CaptureStackBackTrace(uint32 FrameToSkip, uint32 FrameToCapture, void** Backtrace);
+		INSTRUMENTATION_FUNCTION_ATTRIBUTES uint16 CaptureStackBackTrace(uint32 FrameToSkip, uint32 FrameToCapture, void** Backtrace);
 		// Allocates a Tls index.
-		uint32 AllocTlsSlot();
+		INSTRUMENTATION_FUNCTION_ATTRIBUTES uint32 AllocTlsSlot();
 		// Releases a Tls index.
-		void FreeTlsSlot(uint32 Index);
+		INSTRUMENTATION_FUNCTION_ATTRIBUTES void FreeTlsSlot(uint32 Index);
 		// Gets the value of the Tls index for the current thread.
-		void* GetTlsValue(uint32 Index);
+		INSTRUMENTATION_FUNCTION_ATTRIBUTES void* GetTlsValue(uint32 Index);
 		// Sets the value of the Tls index for the current thread.
-		void SetTlsValue(uint32 Index, void* Value);
+		INSTRUMENTATION_FUNCTION_ATTRIBUTES void SetTlsValue(uint32 Index, void* Value);
 		// Returns the current thread Id.
-		uint32 GetCurrentThreadId();
+		INSTRUMENTATION_FUNCTION_ATTRIBUTES uint32 GetCurrentThreadId();
 
 		// Rewrites the patchable function prefix with the proper jmp to the target.
-		void PrepareTrampoline(void* PatchableFunctionAddress, void* TargetFunctionAddress, bool bUseRETBypass);
+		INSTRUMENTATION_FUNCTION_ATTRIBUTES void PrepareTrampoline(void* PatchableFunctionAddress, void* TargetFunctionAddress, bool bUseRETBypass);
 		// Rewrites the patchable function prefix back to its original compiled NOPs.
-		void CleanupTrampoline(void* PatchableFunctionAddress);
+		INSTRUMENTATION_FUNCTION_ATTRIBUTES void CleanupTrampoline(void* PatchableFunctionAddress);
 		// Rewrites the 2 first bytes of the function to jump at the beginning of the prefix section.
-		void ActivateTrampoline(void* PatchableFunctionAddress);
+		INSTRUMENTATION_FUNCTION_ATTRIBUTES void ActivateTrampoline(void* PatchableFunctionAddress);
 		// Rewrites the 2 first bytes of the function to either do nothing or do an immediate return (RET bypass).
-		void DeactivateTrampoline(void* PatchableFunctionAddress, bool bUseRETBypass);
+		INSTRUMENTATION_FUNCTION_ATTRIBUTES void DeactivateTrampoline(void* PatchableFunctionAddress, bool bUseRETBypass);
 		// This needs to be called after trampoline activation or deactivation to make sure it takes effect immediately.
-		void FlushInstructionCache();
+		INSTRUMENTATION_FUNCTION_ATTRIBUTES void FlushInstructionCache();
 		// Get the limits of the stack for the current thread.
-		void GetCurrentThreadStackLimits(void** LowLimit, void** HighLimit);
+		INSTRUMENTATION_FUNCTION_ATTRIBUTES void GetCurrentThreadStackLimits(void** LowLimit, void** HighLimit);
 		// Check if the thread id given is currently alive
-		bool IsThreadAlive(uint32 ThreadId);
+		INSTRUMENTATION_FUNCTION_ATTRIBUTES bool IsThreadAlive(uint32 ThreadId);
 		// Provides a compilation barrier
-		FORCEINLINE void AsymmetricThreadFenceLight()
+		INSTRUMENTATION_FUNCTION_ATTRIBUTES FORCEINLINE void AsymmetricThreadFenceLight()
 		{
 			_ReadWriteBarrier();
 		}
 
 		// The function generates an interprocessor interrupt (IPI) to all processors that are part of the current process affinity.
 		// It guarantees the visibility of write operations performed on one processor to the other processors.
-		void AsymmetricThreadFenceHeavy();
+		INSTRUMENTATION_FUNCTION_ATTRIBUTES void AsymmetricThreadFenceHeavy();
 
 		// Linear allocator for sync objects — reserves virtual address space, commits in chunks,
 		// decommits the entire region in a single syscall on shadow reset.
-		void  InitSyncObjectAllocator(SIZE_T ReservedSize, SIZE_T CommitGranularity);
-		void  ShutdownSyncObjectAllocator();
-		void* SyncObjectAllocatorAlloc(SIZE_T Size);
-		void  SyncObjectAllocatorReset();
-		void  SyncObjectAllocatorAddRef();
-		void  SyncObjectAllocatorRelease();
-		int32 SyncObjectAllocatorGetRefCount();
-		SIZE_T SyncObjectAllocatorGetCommittedSize();
+		INSTRUMENTATION_FUNCTION_ATTRIBUTES void  InitSyncObjectAllocator(SIZE_T ReservedSize, SIZE_T CommitGranularity);
+		INSTRUMENTATION_FUNCTION_ATTRIBUTES void  ShutdownSyncObjectAllocator();
+		INSTRUMENTATION_FUNCTION_ATTRIBUTES void* SyncObjectAllocatorAlloc(SIZE_T Size);
+		INSTRUMENTATION_FUNCTION_ATTRIBUTES void  SyncObjectAllocatorReset();
+		INSTRUMENTATION_FUNCTION_ATTRIBUTES void  SyncObjectAllocatorAddRef();
+		INSTRUMENTATION_FUNCTION_ATTRIBUTES void  SyncObjectAllocatorRelease();
+		INSTRUMENTATION_FUNCTION_ATTRIBUTES int32 SyncObjectAllocatorGetRefCount();
+		INSTRUMENTATION_FUNCTION_ATTRIBUTES SIZE_T SyncObjectAllocatorGetCommittedSize();
 	}
 }
 
diff --git a/Engine/Source/Runtime/Core/Internal/Sanitizer/RaceDetectorTypes.h b/Engine/Source/Runtime/Core/Internal/Sanitizer/RaceDetectorTypes.h
--- a/Engine/Source/Runtime/Core/Internal/Sanitizer/RaceDetectorTypes.h
+++ b/Engine/Source/Runtime/Core/Internal/Sanitizer/RaceDetectorTypes.h
@@ -564,7 +564,7 @@ namespace UE::Sanitizer::RaceDetector {
 			EmplaceEntry<FHistoryEntryFunctionExit>();
 		}
 
-		inline INSTRUMENTATION_FUNCTION_ATTRIBUTES bool ResolveAccess(struct FContext& Context, const void* InAlignedPointer, FMemoryAccess InAccess, FCallstackLocation& OutLocation, FClockRange& OutClockRange);
+		INSTRUMENTATION_FUNCTION_ATTRIBUTES inline bool ResolveAccess(struct FContext& Context, const void* InAlignedPointer, FMemoryAccess InAccess, FCallstackLocation& OutLocation, FClockRange& OutClockRange);
 
 		INSTRUMENTATION_FUNCTION_ATTRIBUTES int32 GetOffset() const
 		{
@@ -740,7 +740,7 @@ namespace UE::Sanitizer::RaceDetector {
 		return false;
 	}
 
-	CORE_API extern FContext* GetThreadContext();
+	INSTRUMENTATION_FUNCTION_ATTRIBUTES CORE_API extern FContext* GetThreadContext();
 
 	struct FInstrumentationScope
 	{
diff --git a/Engine/Source/Runtime/Core/Internal/Sanitizer/Types.h b/Engine/Source/Runtime/Core/Internal/Sanitizer/Types.h
--- a/Engine/Source/Runtime/Core/Internal/Sanitizer/Types.h
+++ b/Engine/Source/Runtime/Core/Internal/Sanitizer/Types.h
@@ -15,14 +15,14 @@
 #include "Hash/CityHash.h"
 
 #define INSTRUMENTATION_ENUM_CLASS_FLAGS(Enum) \
-	FORCEINLINE INSTRUMENTATION_FUNCTION_ATTRIBUTES           Enum& operator|=(Enum& Lhs, Enum Rhs) { return Lhs = (Enum)((__underlying_type(Enum))Lhs | (__underlying_type(Enum))Rhs); } \
-	FORCEINLINE INSTRUMENTATION_FUNCTION_ATTRIBUTES           Enum& operator&=(Enum& Lhs, Enum Rhs) { return Lhs = (Enum)((__underlying_type(Enum))Lhs & (__underlying_type(Enum))Rhs); } \
-	FORCEINLINE INSTRUMENTATION_FUNCTION_ATTRIBUTES           Enum& operator^=(Enum& Lhs, Enum Rhs) { return Lhs = (Enum)((__underlying_type(Enum))Lhs ^ (__underlying_type(Enum))Rhs); } \
-	FORCEINLINE INSTRUMENTATION_FUNCTION_ATTRIBUTES constexpr Enum  operator| (Enum  Lhs, Enum Rhs) { return (Enum)((__underlying_type(Enum))Lhs | (__underlying_type(Enum))Rhs); } \
-	FORCEINLINE INSTRUMENTATION_FUNCTION_ATTRIBUTES constexpr Enum  operator& (Enum  Lhs, Enum Rhs) { return (Enum)((__underlying_type(Enum))Lhs & (__underlying_type(Enum))Rhs); } \
-	FORCEINLINE INSTRUMENTATION_FUNCTION_ATTRIBUTES constexpr Enum  operator^ (Enum  Lhs, Enum Rhs) { return (Enum)((__underlying_type(Enum))Lhs ^ (__underlying_type(Enum))Rhs); } \
-	FORCEINLINE INSTRUMENTATION_FUNCTION_ATTRIBUTES constexpr bool  operator! (Enum  E)             { return !(__underlying_type(Enum))E; } \
-	FORCEINLINE INSTRUMENTATION_FUNCTION_ATTRIBUTES constexpr Enum  operator~ (Enum  E)             { return (Enum)~(__underlying_type(Enum))E; }
+	INSTRUMENTATION_FUNCTION_ATTRIBUTES FORCEINLINE           Enum& operator|=(Enum& Lhs, Enum Rhs) { return Lhs = (Enum)((__underlying_type(Enum))Lhs | (__underlying_type(Enum))Rhs); } \
+	INSTRUMENTATION_FUNCTION_ATTRIBUTES FORCEINLINE           Enum& operator&=(Enum& Lhs, Enum Rhs) { return Lhs = (Enum)((__underlying_type(Enum))Lhs & (__underlying_type(Enum))Rhs); } \
+	INSTRUMENTATION_FUNCTION_ATTRIBUTES FORCEINLINE           Enum& operator^=(Enum& Lhs, Enum Rhs) { return Lhs = (Enum)((__underlying_type(Enum))Lhs ^ (__underlying_type(Enum))Rhs); } \
+	INSTRUMENTATION_FUNCTION_ATTRIBUTES FORCEINLINE constexpr Enum  operator| (Enum  Lhs, Enum Rhs) { return (Enum)((__underlying_type(Enum))Lhs | (__underlying_type(Enum))Rhs); } \
+	INSTRUMENTATION_FUNCTION_ATTRIBUTES FORCEINLINE constexpr Enum  operator& (Enum  Lhs, Enum Rhs) { return (Enum)((__underlying_type(Enum))Lhs & (__underlying_type(Enum))Rhs); } \
+	INSTRUMENTATION_FUNCTION_ATTRIBUTES FORCEINLINE constexpr Enum  operator^ (Enum  Lhs, Enum Rhs) { return (Enum)((__underlying_type(Enum))Lhs ^ (__underlying_type(Enum))Rhs); } \
+	INSTRUMENTATION_FUNCTION_ATTRIBUTES FORCEINLINE constexpr bool  operator! (Enum  E)             { return !(__underlying_type(Enum))E; } \
+	INSTRUMENTATION_FUNCTION_ATTRIBUTES FORCEINLINE constexpr Enum  operator~ (Enum  E)             { return (Enum)~(__underlying_type(Enum))E; }
 
 namespace UE::Sanitizer {
 
@@ -161,7 +161,7 @@ namespace UE::Sanitizer {
 				// we end up decoding inline frames for what comes after we return from the function, which doesn't make sense.
 				// By using Frame - 1, we make sure that we actually point inside the chunk of assembly that represents the function call and it is enough
 				// to make symbol decode correct since we only need to point inside a valid range of a symbol entry, not directly at the start of it.
-				FPlatformStackWalk::EnumerateSymbolInfosForProgramCounter(Frame - 1, bIncludeInlineFrames, 
+				FPlatformStackWalk::EnumerateSymbolInfosForProgramCounter(Frame - 1, bIncludeInlineFrames,
 					[&Loc, Frame](FProgramCounterSymbolInfo& SymbolInfo) INSTRUMENTATION_FUNCTION_ATTRIBUTES
 					{
 						Loc.Locations.Emplace(FString(SymbolInfo.ModuleName), FString(SymbolInfo.FunctionName), FString(SymbolInfo.Filename), SymbolInfo.LineNumber, Frame);
diff --git a/Engine/Source/Runtime/Core/Private/Sanitizer/RaceDetector.cpp b/Engine/Source/Runtime/Core/Private/Sanitizer/RaceDetector.cpp
--- a/Engine/Source/Runtime/Core/Private/Sanitizer/RaceDetector.cpp
+++ b/Engine/Source/Runtime/Core/Private/Sanitizer/RaceDetector.cpp
@@ -171,12 +171,12 @@ namespace UE::Sanitizer::RaceDetector {
 		return !!RaceDetectorActive && !Context.WinInstrumentationDepth && !Context.InstrumentationDepth && !bIsResettingShadow;
 	}
 
-	FORCEINLINE INSTRUMENTATION_FUNCTION_ATTRIBUTES void* GetTlsValueFast(uint32 Value)
+	INSTRUMENTATION_FUNCTION_ATTRIBUTES FORCEINLINE void* GetTlsValueFast(uint32 Value)
 	{
 		return (void*)__readgsqword(Value);
 	}
 
-	FORCEINLINE INSTRUMENTATION_FUNCTION_ATTRIBUTES FContext* GetThreadContext()
+	INSTRUMENTATION_FUNCTION_ATTRIBUTES FORCEINLINE FContext* GetThreadContext()
 	{
 		return (FContext*)GetTlsValueFast(ContextTLSValue);
 	}
@@ -194,12 +194,12 @@ namespace UE::Sanitizer::RaceDetector {
 	// ------------------------------------------------------------------------------
 	// Callstack management.
 	// ------------------------------------------------------------------------------
-	FORCEINLINE INSTRUMENTATION_FUNCTION_ATTRIBUTES FCallstackLocation GetCurrentCallstackLocation(FContext& Context)
+	INSTRUMENTATION_FUNCTION_ATTRIBUTES FORCEINLINE FCallstackLocation GetCurrentCallstackLocation(FContext& Context)
 	{
 		return FCallstackLocation(Context.CurrentCallstack, Context.CurrentCallstackSize);
 	}
 
-	FORCEINLINE INSTRUMENTATION_FUNCTION_ATTRIBUTES void AddCallstackFrame(FContext& Context, void* ReturnAddress)
+	INSTRUMENTATION_FUNCTION_ATTRIBUTES FORCEINLINE void AddCallstackFrame(FContext& Context, void* ReturnAddress)
 	{
 		// History tracing needs to come first since initializing a new history chunk
 		// needs to copy the current stack and we don't want the current stack
@@ -214,7 +214,7 @@ namespace UE::Sanitizer::RaceDetector {
 		Context.CurrentCallstack[Context.CurrentCallstackSize++] = ReturnAddress;
 	}
 
-	FORCEINLINE INSTRUMENTATION_FUNCTION_ATTRIBUTES void RemoveCallstackFrame(FContext& Context)
+	INSTRUMENTATION_FUNCTION_ATTRIBUTES FORCEINLINE void RemoveCallstackFrame(FContext& Context)
 	{
 		// History tracing needs to come first since initializing a new history chunk
 		// needs to copy the current stack before applying the exit.
@@ -321,17 +321,17 @@ namespace UE::Sanitizer::RaceDetector {
 		return nullptr;
 	}
 
-	FORCEINLINE INSTRUMENTATION_FUNCTION_ATTRIBUTES UPTRINT GetAlignedAddress(UPTRINT Ptr)
+	INSTRUMENTATION_FUNCTION_ATTRIBUTES FORCEINLINE UPTRINT GetAlignedAddress(UPTRINT Ptr)
 	{
 		return UPTRINT(Ptr) & ~UPTRINT(0b111);
 	}
 
-	FORCEINLINE INSTRUMENTATION_FUNCTION_ATTRIBUTES void* GetAlignedPointer(void* Ptr)
+	INSTRUMENTATION_FUNCTION_ATTRIBUTES FORCEINLINE void* GetAlignedPointer(void* Ptr)
 	{
 		return (void*)GetAlignedAddress((UPTRINT)Ptr);
 	}
 
-	FORCEINLINE INSTRUMENTATION_FUNCTION_ATTRIBUTES UPTRINT GetShadowMemoryAddress(UPTRINT Ptr)
+	INSTRUMENTATION_FUNCTION_ATTRIBUTES FORCEINLINE UPTRINT GetShadowMemoryAddress(UPTRINT Ptr)
 	{
 		static_assert(sizeof(FShadowMemory) <= sizeof(UPTRINT) * 4);
 
@@ -359,7 +359,7 @@ namespace UE::Sanitizer::RaceDetector {
 		return (FShadowMemory*)ShadowMemoryAddress;
 	}
 	
-	FORCEINLINE INSTRUMENTATION_FUNCTION_ATTRIBUTES UPTRINT GetShadowClockBankMemoryAddress(UPTRINT Ptr)
+	INSTRUMENTATION_FUNCTION_ATTRIBUTES FORCEINLINE UPTRINT GetShadowClockBankMemoryAddress(UPTRINT Ptr)
 	{
 		static_assert(sizeof(FShadowClockBankSlot) <= sizeof(UPTRINT));
 
@@ -454,14 +454,14 @@ namespace UE::Sanitizer::RaceDetector {
 		}
 	}
 	
-	FORCEINLINE INSTRUMENTATION_FUNCTION_ATTRIBUTES uint8 GetOffsetInBytes(void* Ptr)
+	INSTRUMENTATION_FUNCTION_ATTRIBUTES FORCEINLINE uint8 GetOffsetInBytes(void* Ptr)
 	{
 		UPTRINT OriginalAddr = (UPTRINT)Ptr;
 		UPTRINT FinalAddr = OriginalAddr & ~UPTRINT(0b111);
 		return (uint8)(OriginalAddr - FinalAddr);
 	}
 
-	FORCEINLINE INSTRUMENTATION_FUNCTION_ATTRIBUTES void ResetShadowMemory()
+	INSTRUMENTATION_FUNCTION_ATTRIBUTES FORCEINLINE void ResetShadowMemory()
 	{
 		Platform::UnmapShadowMemory();
 	}
@@ -669,12 +669,12 @@ namespace UE::Sanitizer::RaceDetector {
 		}
 	}
 
-	INSTRUMENTATION_FUNCTION_ATTRIBUTES void SetRaceCallbackFn(TRaceCallbackFn CallbackFn)
+	INSTRUMENTATION_ENTRYPOINT_ATTRIBUTES void SetRaceCallbackFn(TRaceCallbackFn CallbackFn)
 	{
 		RaceCallbackFn = CallbackFn;
 	}
 
-	INSTRUMENTATION_FUNCTION_ATTRIBUTES void ResetRaceCallbackFn()
+	INSTRUMENTATION_ENTRYPOINT_ATTRIBUTES void ResetRaceCallbackFn()
 	{
 		RaceCallbackFn.Reset();
 	}
@@ -788,7 +788,7 @@ namespace UE::Sanitizer::RaceDetector {
 
 			// Check if there is any overlap with the previous access
 			// Each bit in Access represets 1-byte in memory so we can just
-			// AND both access to know if we have common bytes being accessed. 
+			// AND both access to know if we have common bytes being accessed.
 			if (!(CurrentAccess.Access & PreviousAccess.Access))
 			{
 				continue;
@@ -934,7 +934,7 @@ namespace UE::Sanitizer::RaceDetector {
 		}
 	}
 
-	INSTRUMENTATION_FUNCTION_ATTRIBUTES void DumpContext()
+	INSTRUMENTATION_ENTRYPOINT_ATTRIBUTES void DumpContext()
 	{
 		FInstrumentationScope InstrumentationScope;
 		FContext& Context = EnsureCurrentContext();
@@ -957,7 +957,7 @@ namespace UE::Sanitizer::RaceDetector {
 			*OtherClocks.ToString());
 	}
 
-	INSTRUMENTATION_FUNCTION_ATTRIBUTES void DumpContextDetailed()
+	INSTRUMENTATION_ENTRYPOINT_ATTRIBUTES void DumpContextDetailed()
 	{
 		FInstrumentationScope InstrumentationScope;
 		FContext& Context = EnsureCurrentContext();
@@ -1163,14 +1163,14 @@ namespace UE::Sanitizer::RaceDetector {
 		HandleReports();
 	}
 	
-	namespace Platform { extern void ThrottleMemoryAccessesDuringShadowReset(); }
+	namespace Platform { INSTRUMENTATION_FUNCTION_ATTRIBUTES extern void ThrottleMemoryAccessesDuringShadowReset(); }
 
 	/**
 	 * FMalloc proxy removes instrumentation for allocators 
 	 * and free the shadow range to avoid detecting races in freed
 	 * memory.
 	 */
-	class FMallocInstrumentation : public FMalloc
+	class AUTORTFM_DISABLE FMallocInstrumentation : public FMalloc
 	{
 	private:
 		/** Malloc we're based on, aka using under the hood */
@@ -1619,12 +1619,12 @@ namespace UE::Sanitizer::RaceDetector {
 		ToggleRaceDetectionInternal(true, .0f, Timeout);
 	}
 
-	INSTRUMENTATION_FUNCTION_ATTRIBUTES void ToggleRaceDetection(bool bEnable, float GracePeriodSeconds)
+	INSTRUMENTATION_ENTRYPOINT_ATTRIBUTES void ToggleRaceDetection(bool bEnable, float GracePeriodSeconds)
 	{
 		ToggleRaceDetectionInternal(bEnable, GracePeriodSeconds);
 	}
 
-	INSTRUMENTATION_FUNCTION_ATTRIBUTES void ToggleThreadDetailedLog(bool bEnabled)
+	INSTRUMENTATION_ENTRYPOINT_ATTRIBUTES void ToggleThreadDetailedLog(bool bEnabled)
 	{
 		FContext& Context = EnsureCurrentContext();
 		if (bEnabled)
@@ -1637,12 +1637,12 @@ namespace UE::Sanitizer::RaceDetector {
 		}
 	}
 
-	INSTRUMENTATION_FUNCTION_ATTRIBUTES void ToggleGlobalDetailedLog(bool bEnabled)
+	INSTRUMENTATION_ENTRYPOINT_ATTRIBUTES void ToggleGlobalDetailedLog(bool bEnabled)
 	{
 		bDetailedLogGlobal = bEnabled;
 	}
 
-	INSTRUMENTATION_FUNCTION_ATTRIBUTES void ToggleFilterDetailedLogOnAddress(void* Address)
+	INSTRUMENTATION_ENTRYPOINT_ATTRIBUTES void ToggleFilterDetailedLogOnAddress(void* Address)
 	{
 		DetailedLogAddress = Address;
 	}
@@ -1667,7 +1667,7 @@ using namespace UE::Sanitizer::RaceDetector;
 #define FINISH_HANDLE_INSTRUMENTATION() \
 	Context.InstrumentationDepth--;
 
-#define INSTRUMENTATION_CORE_FUNCTION_PREFIX CORE_API INSTRUMENTATION_FUNCTION_ATTRIBUTES
+#define INSTRUMENTATION_CORE_FUNCTION_PREFIX INSTRUMENTATION_FUNCTION_ATTRIBUTES CORE_API
 
 extern "C" {
 
@@ -1817,7 +1817,7 @@ extern "C" {
 }
 
 template <typename AtomicType>
-FORCEINLINE INSTRUMENTATION_FUNCTION_ATTRIBUTES AtomicType InstrumentAtomicLoad(AtomicType* Atomic, FAtomicMemoryOrder Order)
+INSTRUMENTATION_FUNCTION_ATTRIBUTES FORCEINLINE AtomicType InstrumentAtomicLoad(AtomicType* Atomic, FAtomicMemoryOrder Order)
 {
 	AtomicType Ret;
 	auto AtomicOp = [Atomic, &Ret]() INSTRUMENTATION_FUNCTION_ATTRIBUTES {
@@ -1836,7 +1836,7 @@ FORCEINLINE INSTRUMENTATION_FUNCTION_ATTRIBUTES AtomicType InstrumentAtomicLoad(
 }
 
 template <typename AtomicType>
-FORCEINLINE INSTRUMENTATION_FUNCTION_ATTRIBUTES void InstrumentAtomicStore(AtomicType* Atomic, AtomicType Val, FAtomicMemoryOrder Order)
+INSTRUMENTATION_FUNCTION_ATTRIBUTES FORCEINLINE void InstrumentAtomicStore(AtomicType* Atomic, AtomicType Val, FAtomicMemoryOrder Order)
 {
 	auto AtomicOp = [Atomic, Val]() INSTRUMENTATION_FUNCTION_ATTRIBUTES {
 		UE::Core::Private::Atomic::Store(Atomic, Val);
@@ -1853,7 +1853,7 @@ FORCEINLINE INSTRUMENTATION_FUNCTION_ATTRIBUTES void InstrumentAtomicStore(Atomi
 }
 
 template <typename AtomicType>
-FORCEINLINE INSTRUMENTATION_FUNCTION_ATTRIBUTES AtomicType InstrumentAtomicExchange(AtomicType* Atomic, AtomicType Val, FAtomicMemoryOrder Order)
+INSTRUMENTATION_FUNCTION_ATTRIBUTES FORCEINLINE AtomicType InstrumentAtomicExchange(AtomicType* Atomic, AtomicType Val, FAtomicMemoryOrder Order)
 {
 	AtomicType Ret;
 	auto AtomicOp = [Atomic, Val, &Ret]() INSTRUMENTATION_FUNCTION_ATTRIBUTES {
@@ -1872,7 +1872,7 @@ FORCEINLINE INSTRUMENTATION_FUNCTION_ATTRIBUTES AtomicType InstrumentAtomicExcha
 }
 
 template <typename AtomicType>
-FORCEINLINE INSTRUMENTATION_FUNCTION_ATTRIBUTES AtomicType InstrumentAtomicCompareExchange(AtomicType* Atomic, AtomicType* Expected, AtomicType Val, FAtomicMemoryOrder SuccessOrder, FAtomicMemoryOrder FailureOrder)
+INSTRUMENTATION_FUNCTION_ATTRIBUTES FORCEINLINE AtomicType InstrumentAtomicCompareExchange(AtomicType* Atomic, AtomicType* Expected, AtomicType Val, FAtomicMemoryOrder SuccessOrder, FAtomicMemoryOrder FailureOrder)
 {
 	AtomicType Ret;
 	auto AtomicOp = [Atomic, Expected, Val, &Ret]() INSTRUMENTATION_FUNCTION_ATTRIBUTES -> bool {
@@ -1897,7 +1897,7 @@ FORCEINLINE INSTRUMENTATION_FUNCTION_ATTRIBUTES AtomicType InstrumentAtomicCompa
 }
 
 template <typename AtomicType>
-FORCEINLINE INSTRUMENTATION_FUNCTION_ATTRIBUTES AtomicType InstrumentAtomicFetchAdd(AtomicType* Atomic, AtomicType Val, FAtomicMemoryOrder Order)
+INSTRUMENTATION_FUNCTION_ATTRIBUTES FORCEINLINE AtomicType InstrumentAtomicFetchAdd(AtomicType* Atomic, AtomicType Val, FAtomicMemoryOrder Order)
 {
 	AtomicType Ret;
 	auto AtomicOp = [Atomic, Val, &Ret]() INSTRUMENTATION_FUNCTION_ATTRIBUTES {
@@ -1916,7 +1916,7 @@ FORCEINLINE INSTRUMENTATION_FUNCTION_ATTRIBUTES AtomicType InstrumentAtomicFetch
 }
 
 template <typename AtomicType>
-FORCEINLINE INSTRUMENTATION_FUNCTION_ATTRIBUTES AtomicType InstrumentAtomicFetchSub(AtomicType* Atomic, AtomicType Val, FAtomicMemoryOrder Order)
+INSTRUMENTATION_FUNCTION_ATTRIBUTES FORCEINLINE AtomicType InstrumentAtomicFetchSub(AtomicType* Atomic, AtomicType Val, FAtomicMemoryOrder Order)
 {
 	AtomicType Ret;
 	auto AtomicOp = [Atomic, Val, &Ret]() INSTRUMENTATION_FUNCTION_ATTRIBUTES {
@@ -1935,7 +1935,7 @@ FORCEINLINE INSTRUMENTATION_FUNCTION_ATTRIBUTES AtomicType InstrumentAtomicFetch
 }
 
 template <typename AtomicType>
-FORCEINLINE INSTRUMENTATION_FUNCTION_ATTRIBUTES AtomicType InstrumentAtomicFetchOr(AtomicType* Atomic, AtomicType Val, FAtomicMemoryOrder Order)
+INSTRUMENTATION_FUNCTION_ATTRIBUTES FORCEINLINE AtomicType InstrumentAtomicFetchOr(AtomicType* Atomic, AtomicType Val, FAtomicMemoryOrder Order)
 {
 	AtomicType Ret;
 	auto AtomicOp = [Atomic, Val, &Ret]() INSTRUMENTATION_FUNCTION_ATTRIBUTES {
@@ -1954,7 +1954,7 @@ FORCEINLINE INSTRUMENTATION_FUNCTION_ATTRIBUTES AtomicType InstrumentAtomicFetch
 }
 
 template <typename AtomicType>
-FORCEINLINE INSTRUMENTATION_FUNCTION_ATTRIBUTES AtomicType InstrumentAtomicFetchXor(AtomicType* Atomic, AtomicType Val, FAtomicMemoryOrder Order)
+INSTRUMENTATION_FUNCTION_ATTRIBUTES FORCEINLINE AtomicType InstrumentAtomicFetchXor(AtomicType* Atomic, AtomicType Val, FAtomicMemoryOrder Order)
 {
 	AtomicType Ret;
 	auto AtomicOp = [Atomic, Val, &Ret]() INSTRUMENTATION_FUNCTION_ATTRIBUTES {
@@ -1973,7 +1973,7 @@ FORCEINLINE INSTRUMENTATION_FUNCTION_ATTRIBUTES AtomicType InstrumentAtomicFetch
 }
 
 template <typename AtomicType>
-FORCEINLINE INSTRUMENTATION_FUNCTION_ATTRIBUTES AtomicType InstrumentAtomicFetchAnd(AtomicType* Atomic, AtomicType Val, FAtomicMemoryOrder Order)
+INSTRUMENTATION_FUNCTION_ATTRIBUTES FORCEINLINE AtomicType InstrumentAtomicFetchAnd(AtomicType* Atomic, AtomicType Val, FAtomicMemoryOrder Order)
 {
 	AtomicType Ret;
 	auto AtomicOp = [Atomic, Val, &Ret]() INSTRUMENTATION_FUNCTION_ATTRIBUTES {
diff --git a/Engine/Source/Runtime/Core/Public/Clang/ClangPlatformCodeAnalysis.h b/Engine/Source/Runtime/Core/Public/Clang/ClangPlatformCodeAnalysis.h
--- a/Engine/Source/Runtime/Core/Public/Clang/ClangPlatformCodeAnalysis.h
+++ b/Engine/Source/Runtime/Core/Public/Clang/ClangPlatformCodeAnalysis.h
@@ -144,14 +144,22 @@
 #endif
 
 	// Thread-sanitiser annotation functions.
+	#if defined(__AUTORTFM) && __AUTORTFM
+		// both open and closed code call directly, no closed variant or trampoline generated.
+		#include "AutoRTFM/Constants.h"
+		#define ANNOTATE_FUNC_ATTRIBUTES __attribute__((autortfm(autortfm_mode_internal)))
+	#else
+		#define ANNOTATE_FUNC_ATTRIBUTES
+	#endif
 	#ifdef __cplusplus
 	extern "C" {
 	#endif
-		void AnnotateHappensBefore(const char *f, int l, void *addr);
-		void AnnotateHappensAfter(const char *f, int l, void *addr);
+		ANNOTATE_FUNC_ATTRIBUTES void AnnotateHappensBefore(const char *f, int l, void *addr);
+		ANNOTATE_FUNC_ATTRIBUTES void AnnotateHappensAfter(const char *f, int l, void *addr);
 	#ifdef __cplusplus
 	}
 	#endif
+	#undef ANNOTATE_FUNC_ATTRIBUTES
 
 	// Annotate that previous load/stores occur before addr 
 	#define TSAN_BEFORE(addr) AnnotateHappensBefore(__FILE__, __LINE__, (void*)(addr))
diff --git a/Engine/Source/Runtime/Core/Public/Sanitizer/RaceDetector.h b/Engine/Source/Runtime/Core/Public/Sanitizer/RaceDetector.h
--- a/Engine/Source/Runtime/Core/Public/Sanitizer/RaceDetector.h
+++ b/Engine/Source/Runtime/Core/Public/Sanitizer/RaceDetector.h
@@ -6,50 +6,51 @@
 
 #if USING_INSTRUMENTATION
 
+#include "Instrumentation/Defines.h"
 #include "Sanitizer/Types.h"
 #include "Misc/Timeout.h"
 
 namespace UE::Sanitizer::RaceDetector
 {
 	// Setup everything that is required for the race detector.
-	CORE_API bool Initialize();
+	INSTRUMENTATION_ENTRYPOINT_ATTRIBUTES CORE_API bool Initialize();
 	// Cleanup everything about the race detector that was setup in Initialize.
-	CORE_API bool Shutdown();
+	INSTRUMENTATION_ENTRYPOINT_ATTRIBUTES CORE_API bool Shutdown();
 
 	typedef TFunctionRef<void(uint64 RaceAddress, uint32 FirstThreadId, uint32 SecondThreadId, const FFullLocation& FirstLocation, const FFullLocation& SecondLocation)> TRaceCallbackFn;
 	// Set a callback that will be called whenever a race is detected (filters will not apply).
 	// NOTE: It is not safe to perform any allocation within the callback since it will be
 	//       called from instrumentation code.
-	CORE_API void SetRaceCallbackFn(TRaceCallbackFn CallbackFn);
+	INSTRUMENTATION_ENTRYPOINT_ATTRIBUTES CORE_API void SetRaceCallbackFn(TRaceCallbackFn CallbackFn);
 	// Remove any callback that was previously set.
-	CORE_API void ResetRaceCallbackFn();
+	INSTRUMENTATION_ENTRYPOINT_ATTRIBUTES CORE_API void ResetRaceCallbackFn();
 
 	// Dump the current thread context to help with debugging.
-	CORE_API void DumpContext();
+	INSTRUMENTATION_ENTRYPOINT_ATTRIBUTES CORE_API void DumpContext();
 	// Dump the current thread context but with additional info.
-	CORE_API void DumpContextDetailed();
+	INSTRUMENTATION_ENTRYPOINT_ATTRIBUTES CORE_API  void DumpContextDetailed();
 
 	// Means that races that doesn't involve the current thread will be filtered.
-	CORE_API void ToggleFilterOtherThreads(bool bEnable);
+	INSTRUMENTATION_ENTRYPOINT_ATTRIBUTES CORE_API void ToggleFilterOtherThreads(bool bEnable);
 
 	// Toggle race detection until a call to the same function that toggles it back off.
 	// If multiple threads activate race detection at the same time, it will be turned off when all thread toggle it back to off.
 	// The grace period is how much time to wait until race detection is really turned off. Can be used when toggling is expected to happen too often or too fast.
-	CORE_API void ToggleRaceDetection(bool bEnable, float GracePeriodSeconds = 0.0f);
+	INSTRUMENTATION_ENTRYPOINT_ATTRIBUTES CORE_API void ToggleRaceDetection(bool bEnable, float GracePeriodSeconds = 0.0f);
 	// Toggle filter on duplicate races. Does not apply when a callback is set.
-	CORE_API void ToggleFilterDuplicateRaces(bool bEnable);
+	INSTRUMENTATION_ENTRYPOINT_ATTRIBUTES CORE_API void ToggleFilterDuplicateRaces(bool bEnable);
 	// Activate race detection for certain time after which it will be automatically turned off.
-	CORE_API void ToggleRaceDetectionUntil(UE::FTimeout Timeout);
+	INSTRUMENTATION_ENTRYPOINT_ATTRIBUTES CORE_API void ToggleRaceDetectionUntil(UE::FTimeout Timeout);
 
 	// Toggle detailed logging globally for all threads. Very slow!
-	CORE_API void ToggleGlobalDetailedLog(bool bEnable);
+	INSTRUMENTATION_ENTRYPOINT_ATTRIBUTES CORE_API void ToggleGlobalDetailedLog(bool bEnable);
 	// Toggle detailed logging for the current thread.
-	CORE_API void ToggleThreadDetailedLog(bool bEnable);
+	INSTRUMENTATION_ENTRYPOINT_ATTRIBUTES CORE_API void ToggleThreadDetailedLog(bool bEnable);
 	// Toggle detailed logging filter to only log activity for a specific address.
 	// NOTE: Set to nullptr to remove filtering.
-	CORE_API void ToggleFilterDetailedLogOnAddress(void* Ptr);
+	INSTRUMENTATION_ENTRYPOINT_ATTRIBUTES CORE_API void ToggleFilterDetailedLogOnAddress(void* Ptr);
 	// Returns whether the race detector is currently active or not.
-	CORE_API bool IsActive();
+	INSTRUMENTATION_ENTRYPOINT_ATTRIBUTES CORE_API bool IsActive();
 	
 	class FRaceDetectorScope
 	{
diff --git a/Engine/Source/Runtime/Instrumentation/Instrumentation.Build.cs b/Engine/Source/Runtime/Instrumentation/Instrumentation.Build.cs
--- a/Engine/Source/Runtime/Instrumentation/Instrumentation.Build.cs
+++ b/Engine/Source/Runtime/Instrumentation/Instrumentation.Build.cs
@@ -14,5 +14,9 @@ public class Instrumentation : ModuleRules
 		PrivateDefinitions.Add("SUPPRESS_PER_MODULE_INLINE_FILE"); // This module does not use core's standard operator new/delete overloads
 		
 		bDisableAutoRTFMInstrumentation = true;
+
+		// Instrumentation symbols (__Thunk__*) must be exported even with merged modules,
+		// since they are dllimport'd by PerModuleInline.inl in every other module/DLL.
+		ModuleSymbolVisibility = SymbolVisibility.VisibileForDll;
 	}
 }
diff --git a/Engine/Source/Runtime/Instrumentation/Internal/Instrumentation/Defines.h b/Engine/Source/Runtime/Instrumentation/Internal/Instrumentation/Defines.h
--- a/Engine/Source/Runtime/Instrumentation/Internal/Instrumentation/Defines.h
+++ b/Engine/Source/Runtime/Instrumentation/Internal/Instrumentation/Defines.h
@@ -2,7 +2,18 @@
 
 #pragma once
 
-#define INSTRUMENTATION_FUNCTION_ATTRIBUTES __declspec(safebuffers) __attribute__((disable_sanitizer_instrumentation))
+#define INSTRUMENTATION_FUNCTION_BASE_ATTRIBUTES __declspec(safebuffers) __attribute__((disable_sanitizer_instrumentation))
+
+#if defined(__AUTORTFM) && __AUTORTFM
+	#include "AutoRTFM/Constants.h"
+	// autortfm_mode_internal are prevented from being inlined, so we only use that attribute on entry points
+	#define INSTRUMENTATION_ENTRYPOINT_ATTRIBUTES INSTRUMENTATION_FUNCTION_BASE_ATTRIBUTES __attribute__((autortfm(autortfm_mode_internal)))
+	// use autortfm disable once we're inside race detection code to allow inlining
+	#define INSTRUMENTATION_FUNCTION_ATTRIBUTES   INSTRUMENTATION_FUNCTION_BASE_ATTRIBUTES __attribute__((autortfm(autortfm_mode_disable)))
+#else
+	#define INSTRUMENTATION_FUNCTION_ATTRIBUTES   INSTRUMENTATION_FUNCTION_BASE_ATTRIBUTES
+	#define INSTRUMENTATION_ENTRYPOINT_ATTRIBUTES INSTRUMENTATION_FUNCTION_BASE_ATTRIBUTES
+#endif
 
 // Reserves 12 NOPs before the function and 2 NOPs at function entry so that we can patch atomically
 #define INSTRUMENTATION_HOTPATCH_TOTAL_NOPS  14
diff --git a/Engine/Source/Runtime/Instrumentation/Internal/Instrumentation/EntryPoints.h b/Engine/Source/Runtime/Instrumentation/Internal/Instrumentation/EntryPoints.h
--- a/Engine/Source/Runtime/Instrumentation/Internal/Instrumentation/EntryPoints.h
+++ b/Engine/Source/Runtime/Instrumentation/Internal/Instrumentation/EntryPoints.h
@@ -5,34 +5,35 @@
 #if USING_INSTRUMENTATION
 
 #include "Instrumentation/Types.h"
+#include "Instrumentation/Defines.h"
 
 extern "C"
 {
-	INSTRUMENTATION_API void __Thunk__AnnotateHappensBefore(const char* f, int l, void* addr);
-	INSTRUMENTATION_API void __Thunk__AnnotateHappensAfter(const char* f, int l, void* addr);
-	INSTRUMENTATION_API void __Thunk__Instrument_FuncEntry(void* ReturnAddress);
-	INSTRUMENTATION_API void __Thunk__Instrument_FuncExit();
-	INSTRUMENTATION_API void __Thunk__Instrument_Store(uint64 Address, uint32 Size);
-	INSTRUMENTATION_API void __Thunk__Instrument_Load(uint64 Address, uint32 Size);
-	INSTRUMENTATION_API void __Thunk__Instrument_VPtr_Store(void** Address, void* Value);
-	INSTRUMENTATION_API void __Thunk__Instrument_VPtr_Load(void** Address);
-	INSTRUMENTATION_API void __Thunk__Instrument_StoreRange(uint64 Address, uint32 Size);
-	INSTRUMENTATION_API void __Thunk__Instrument_LoadRange(uint64 Address, uint32 Size);
+	INSTRUMENTATION_FUNCTION_ATTRIBUTES INSTRUMENTATION_API void __Thunk__AnnotateHappensBefore(const char* f, int l, void* addr);
+	INSTRUMENTATION_FUNCTION_ATTRIBUTES INSTRUMENTATION_API void __Thunk__AnnotateHappensAfter(const char* f, int l, void* addr);
+	INSTRUMENTATION_FUNCTION_ATTRIBUTES INSTRUMENTATION_API void __Thunk__Instrument_FuncEntry(void* ReturnAddress);
+	INSTRUMENTATION_FUNCTION_ATTRIBUTES INSTRUMENTATION_API void __Thunk__Instrument_FuncExit();
+	INSTRUMENTATION_FUNCTION_ATTRIBUTES INSTRUMENTATION_API void __Thunk__Instrument_Store(uint64 Address, uint32 Size);
+	INSTRUMENTATION_FUNCTION_ATTRIBUTES INSTRUMENTATION_API void __Thunk__Instrument_Load(uint64 Address, uint32 Size);
+	INSTRUMENTATION_FUNCTION_ATTRIBUTES INSTRUMENTATION_API void __Thunk__Instrument_VPtr_Store(void** Address, void* Value);
+	INSTRUMENTATION_FUNCTION_ATTRIBUTES INSTRUMENTATION_API void __Thunk__Instrument_VPtr_Load(void** Address);
+	INSTRUMENTATION_FUNCTION_ATTRIBUTES INSTRUMENTATION_API void __Thunk__Instrument_StoreRange(uint64 Address, uint32 Size);
+	INSTRUMENTATION_FUNCTION_ATTRIBUTES INSTRUMENTATION_API void __Thunk__Instrument_LoadRange(uint64 Address, uint32 Size);
 
 	#define INSTRUMENT_LOAD_FUNC(Type) \
-	INSTRUMENTATION_API Type __Thunk__Instrument_AtomicLoad_##Type(Type* Atomic, UE::Instrumentation::FAtomicMemoryOrder MemoryOrder);
+	INSTRUMENTATION_FUNCTION_ATTRIBUTES INSTRUMENTATION_API Type __Thunk__Instrument_AtomicLoad_##Type(Type* Atomic, UE::Instrumentation::FAtomicMemoryOrder MemoryOrder);
 
 	#define INSTRUMENT_STORE_FUNC(Type) \
-	INSTRUMENTATION_API void __Thunk__Instrument_AtomicStore_##Type(Type* Atomic, Type Val, UE::Instrumentation::FAtomicMemoryOrder MemoryOrder);
+	INSTRUMENTATION_FUNCTION_ATTRIBUTES INSTRUMENTATION_API void __Thunk__Instrument_AtomicStore_##Type(Type* Atomic, Type Val, UE::Instrumentation::FAtomicMemoryOrder MemoryOrder);
 
 	#define INSTRUMENT_EXCHANGE_FUNC(Type) \
-	INSTRUMENTATION_API Type __Thunk__Instrument_AtomicExchange_##Type(Type* Atomic, Type Val, UE::Instrumentation::FAtomicMemoryOrder MemoryOrder);
+	INSTRUMENTATION_FUNCTION_ATTRIBUTES INSTRUMENTATION_API Type __Thunk__Instrument_AtomicExchange_##Type(Type* Atomic, Type Val, UE::Instrumentation::FAtomicMemoryOrder MemoryOrder);
 
 	#define INSTRUMENT_COMPARE_EXCHANGE_FUNC(Type) \
-	INSTRUMENTATION_API Type __Thunk__Instrument_AtomicCompareExchange_##Type(Type* Atomic, Type* Expected, Type Val, UE::Instrumentation::FAtomicMemoryOrder SuccessMemoryOrder, UE::Instrumentation::FAtomicMemoryOrder FailureMemoryOrder);
+	INSTRUMENTATION_FUNCTION_ATTRIBUTES INSTRUMENTATION_API Type __Thunk__Instrument_AtomicCompareExchange_##Type(Type* Atomic, Type* Expected, Type Val, UE::Instrumentation::FAtomicMemoryOrder SuccessMemoryOrder, UE::Instrumentation::FAtomicMemoryOrder FailureMemoryOrder);
 
 	#define INSTRUMENT_RMW_FUNC(Func, Type) \
-	INSTRUMENTATION_API Type __Thunk__Instrument_Atomic##Func##_##Type(Type* Atomic, Type Val, UE::Instrumentation::FAtomicMemoryOrder MemoryOrder);
+	INSTRUMENTATION_FUNCTION_ATTRIBUTES INSTRUMENTATION_API Type __Thunk__Instrument_Atomic##Func##_##Type(Type* Atomic, Type Val, UE::Instrumentation::FAtomicMemoryOrder MemoryOrder);
 
 	// Define the functions with C linkage, referenced by the compiler's instrumentation pass.
 	INSTRUMENT_LOAD_FUNC(int8)
diff --git a/Engine/Source/Runtime/Instrumentation/Internal/Instrumentation/PerModuleInline.inl b/Engine/Source/Runtime/Instrumentation/Internal/Instrumentation/PerModuleInline.inl
--- a/Engine/Source/Runtime/Instrumentation/Internal/Instrumentation/PerModuleInline.inl
+++ b/Engine/Source/Runtime/Instrumentation/Internal/Instrumentation/PerModuleInline.inl
@@ -9,9 +9,9 @@
 #include "Instrumentation/EntryPoints.h"
 
 #if defined(USE_INSTRUMENTATION_PERMODULE_HOTPATCH) && USE_INSTRUMENTATION_PERMODULE_HOTPATCH
-	#define INSTRUMENTATION_PERMODULE_FUNCTION_ATTRIBUTES INSTRUMENTATION_FUNCTION_ATTRIBUTES INSTRUMENTATION_FUNCTION_HOTPATCHABLE	
+	#define INSTRUMENTATION_PERMODULE_FUNCTION_ATTRIBUTES INSTRUMENTATION_ENTRYPOINT_ATTRIBUTES INSTRUMENTATION_FUNCTION_HOTPATCHABLE	
 #else
-	#define INSTRUMENTATION_PERMODULE_FUNCTION_ATTRIBUTES INSTRUMENTATION_FUNCTION_ATTRIBUTES	
+	#define INSTRUMENTATION_PERMODULE_FUNCTION_ATTRIBUTES INSTRUMENTATION_ENTRYPOINT_ATTRIBUTES	
 #endif
 
 // Define the functions with C linkage, referenced by the compiler's instrumentation pass.
diff --git a/Engine/Source/Runtime/Instrumentation/Private/Instrumentation/EntryPoints.cpp b/Engine/Source/Runtime/Instrumentation/Private/Instrumentation/EntryPoints.cpp
--- a/Engine/Source/Runtime/Instrumentation/Private/Instrumentation/EntryPoints.cpp
+++ b/Engine/Source/Runtime/Instrumentation/Private/Instrumentation/EntryPoints.cpp
@@ -8,7 +8,7 @@
 #include "Templates/Atomic.h"
 #include "CoreTypes.h"
 
-#define INSTRUMENTATION_THUNK_ATTRIBUTES INSTRUMENTATION_API INSTRUMENTATION_FUNCTION_ATTRIBUTES INSTRUMENTATION_FUNCTION_HOTPATCHABLE
+#define INSTRUMENTATION_THUNK_ATTRIBUTES INSTRUMENTATION_FUNCTION_ATTRIBUTES INSTRUMENTATION_FUNCTION_HOTPATCHABLE INSTRUMENTATION_API
 
 // Define the functions with C linkage, referenced by the compiler's instrumentation pass.
 extern "C" {
