[TraceAux]
- Suppress warning when trying to enable a read-only channel that is already enabled

#Synced-CL 52534307

diff --git a/Engine/Source/Runtime/Core/Private/ProfilingDebugging/TraceAuxiliary.cpp b/Engine/Source/Runtime/Core/Private/ProfilingDebugging/TraceAuxiliary.cpp
--- a/Engine/Source/Runtime/Core/Private/ProfilingDebugging/TraceAuxiliary.cpp
+++ b/Engine/Source/Runtime/Core/Private/ProfilingDebugging/TraceAuxiliary.cpp
@@ -456,12 +456,15 @@ bool FTraceAuxiliaryImpl::EnableChannel(ChannelType ChannelIdentifier, const FTr
 	// It is not possible to change read only channels once trace is initialized.
 	if (bReadOnlyChannelsFrozen && Channel->IsReadOnly())
 	{
-		UE_LOG_REF(
-			LogCategory,
-			Error,
-			TEXT("Channel '%s' is read only. It is not allowed to manually enable this channel."),
-			ChannelName.ToString()
-		);
+		if (!Channel->IsEnabled())
+		{
+			UE_LOG_REF(
+				LogCategory,
+				Error,
+				TEXT("Channel '%s' is read only. It is not allowed to manually enable this channel."),
+				ChannelName.ToString()
+			);
+		}
 		return Channel->IsEnabled();
 	}
 
