Disable PVS warnings for UE_LOGF and UE_CLOGF consistent with UE_LOG/UE_CLOG
#rnx

diff --git a/Engine/Source/Runtime/Core/Public/Analysis/PvsStudioCodeAnalysis.h b/Engine/Source/Runtime/Core/Public/Analysis/PvsStudioCodeAnalysis.h
--- a/Engine/Source/Runtime/Core/Public/Analysis/PvsStudioCodeAnalysis.h
+++ b/Engine/Source/Runtime/Core/Public/Analysis/PvsStudioCodeAnalysis.h
@@ -27,7 +27,9 @@
 	//-V:BEGIN_UPDATE_LOOP:595
 	//-V:OPENGL_PERFORMANCE_DATA_INVALID:560,564
 	//-V:UE_CLOG(:501,560
+	//-V:UE_CLOGF(:501,560
 	//-V:UE_LOG(:501,510,560
+	//-V:UE_LOGF(:501,510,560
 	//-V:ensure(:595
 	//-V:ALLOCATE_VERTEX_DATA_TEMPLATE:501
 	//-V:UGL_REQUIRED:501
