Prevent the reading of garbage values if ::GetHandleInformation fails.

#rnx

- When creating child processes we check if the handle for the pipes can be inherited by the child process and if not then prevent their use.
- If the call to ::GetHandleInformation failed we record thise via an ensure but allow the code flow to contiue, including checking the flag for HANDLE_FLAG_INHERIT but if the call failed then the Flags variable would not be written to and since it was not being initialized we could be checking against anything.
- We now initialize the Flags variable to 0 so if the call to ::GetHandleInformation fails we assume that the handle cannot be inherited and so it will not be used.
- Moved the code for this check to a stand alone function to avoid code duplication.

#Synced-CL 52535685

diff --git a/Engine/Source/Runtime/Core/Private/Windows/WindowsPlatformProcess.cpp b/Engine/Source/Runtime/Core/Private/Windows/WindowsPlatformProcess.cpp
--- a/Engine/Source/Runtime/Core/Private/Windows/WindowsPlatformProcess.cpp
+++ b/Engine/Source/Runtime/Core/Private/Windows/WindowsPlatformProcess.cpp
@@ -103,6 +103,33 @@ namespace WindowsPlatformProcess
 	}
 }
 
+/**
+ * Checks if a given handle can be inherited by a child process.
+ * 
+ * @param Handle	The handle to check
+ * @param Label		A descriptor used to improve logging if an error is encountered.
+ * 
+ * @return If the handle is valid and can inherited (has the HANDLE_FLAG_INHERIT flag)
+ * then this returns the handle, otherwise it will return nullptr.
+ */
+static void* VerifyHandleIsInheritable(void* Handle, const ANSICHAR* Label)
+{
+	if (!Handle)
+	{
+		return nullptr;
+	}
+
+	DWORD Flags = 0;
+	ensure(GetHandleInformation(HANDLE(Handle), &Flags));
+	if (!(Flags & HANDLE_FLAG_INHERIT))
+	{
+		UE_LOGF(LogWindows, Error, "%s passed to CreateProc is not inheritable.", Label);
+		return nullptr;
+	}
+
+	return Handle;
+}
+
 static bool bCustomProcessAffinity = false;
 bool FWindowsPlatformProcess::SetProcessAffinity(uint32 NumCoresForAffinity, bool bPhysicalCoresOnly)
 {
@@ -482,26 +509,9 @@ void FWindowsPlatformProcess::LaunchURLInternal(const TCHAR* URL, const TCHAR* P
 
 FProcHandle FWindowsPlatformProcess::CreateProc(const TCHAR* URL, const TCHAR* Parms, bool bLaunchDetached, bool bLaunchHidden, bool bLaunchReallyHidden, uint32* OutProcessID, int32 PriorityModifier, const TCHAR* OptionalWorkingDirectory, void* PipeWriteChild, void* PipeReadChild)
 {
-	if (PipeWriteChild)
-	{
-		DWORD Flags;
-		ensure(GetHandleInformation(HANDLE(PipeWriteChild), &Flags));
-		if (!(Flags & HANDLE_FLAG_INHERIT))
-		{
-			UE_LOGF(LogWindows, Error, "PipeWriteChild passed to CreateProc is not inheritable.");
-			PipeWriteChild = nullptr;
-		}
-	}
-	if (PipeReadChild)
-	{
-		DWORD Flags;
-		ensure(GetHandleInformation(HANDLE(PipeReadChild), &Flags));
-		if (!(Flags & HANDLE_FLAG_INHERIT))
-		{
-			UE_LOGF(LogWindows, Error, "PipeReadChild passed to CreateProc is not inheritable.");
-			PipeReadChild = nullptr;
-		}
-	}
+	PipeWriteChild	= VerifyHandleIsInheritable(PipeWriteChild, "PipeWriteChild");
+	PipeReadChild	= VerifyHandleIsInheritable(PipeReadChild, "PipeReadChild");
+
 	// CreateProc used to only have a single "write" pipe argument that both stdout and stderr would be piped into on Windows,
 	// so for this overload we'll preserve that behaviour for compatibility with existing code
 	return CreateProcInternal(URL, Parms, bLaunchDetached, bLaunchHidden, bLaunchReallyHidden, false, OutProcessID, PriorityModifier, OptionalWorkingDirectory, PipeWriteChild, PipeReadChild, PipeWriteChild, false);
@@ -509,36 +519,10 @@ FProcHandle FWindowsPlatformProcess::CreateProc(const TCHAR* URL, const TCHAR* P
 
 FProcHandle FWindowsPlatformProcess::CreateProc(const TCHAR* URL, const TCHAR* Parms, bool bLaunchDetached, bool bLaunchHidden, bool bLaunchReallyHidden, uint32* OutProcessID, int32 PriorityModifier, const TCHAR* OptionalWorkingDirectory, void* PipeWriteChild, void* PipeReadChild, void* PipeStdErrChild)
 {
-	if (PipeWriteChild)
-	{
-		DWORD Flags;
-		ensure(GetHandleInformation(HANDLE(PipeWriteChild), &Flags));
-		if (!(Flags & HANDLE_FLAG_INHERIT))
-		{
-			UE_LOGF(LogWindows, Error, "PipeWriteChild passed to CreateProc is not inheritable.");
-			PipeWriteChild = nullptr;
-		}
-	}
-	if (PipeReadChild)
-	{
-		DWORD Flags;
-		ensure(GetHandleInformation(HANDLE(PipeReadChild), &Flags));
-		if (!(Flags & HANDLE_FLAG_INHERIT))
-		{
-			UE_LOGF(LogWindows, Error, "PipeReadChild passed to CreateProc is not inheritable.");
-			PipeReadChild = nullptr;
-		}
-	}
-	if (PipeStdErrChild)
-	{
-		DWORD Flags;
-		ensure(GetHandleInformation(HANDLE(PipeStdErrChild), &Flags));
-		if (!(Flags & HANDLE_FLAG_INHERIT))
-		{
-			UE_LOGF(LogWindows, Error, "PipeStdErrChild passed to CreateProc is not inheritable.");
-			PipeStdErrChild = nullptr;
-		}
-	}
+	PipeWriteChild	= VerifyHandleIsInheritable(PipeWriteChild, "PipeWriteChild");
+	PipeReadChild	= VerifyHandleIsInheritable(PipeReadChild, "PipeReadChild");
+	PipeStdErrChild	= VerifyHandleIsInheritable(PipeStdErrChild, "PipeStdErrChild");
+
 	return CreateProcInternal(URL, Parms, bLaunchDetached, bLaunchHidden, bLaunchReallyHidden, false, OutProcessID, PriorityModifier, OptionalWorkingDirectory, PipeWriteChild, PipeReadChild, PipeStdErrChild, false);
 }
 
