Add CompareSensitive method to FName that matches FNameEntryId::CompareLexicalSensitive.

#Synced-CL 52471422

diff --git a/Engine/Source/Runtime/Core/Private/UObject/UnrealNames.cpp b/Engine/Source/Runtime/Core/Private/UObject/UnrealNames.cpp
--- a/Engine/Source/Runtime/Core/Private/UObject/UnrealNames.cpp
+++ b/Engine/Source/Runtime/Core/Private/UObject/UnrealNames.cpp
@@ -3507,7 +3507,8 @@ bool FName::Equals(FName Name, const WIDECHAR* Str)
 	return Equals(Name, FWideStringView(Str));
 }
 
-int32 FName::Compare( const FName& Other ) const
+template <ESearchCase::Type SearchCase>
+int32 FName::CompareInternal( const FName& Other ) const
 {
 #if UE_FNAME_OUTLINE_NUMBER
 	if (GetComparisonIndexInternal() == Other.GetComparisonIndexInternal())
@@ -3521,7 +3522,7 @@ int32 FName::Compare( const FName& Other ) const
 	}
 
 	// Names don't match. This means we don't even need to check numbers.
-	return CompareDifferentIdsAlphabetically<ESearchCase::IgnoreCase>(GetComparisonIndex(), Other.GetComparisonIndex());
+	return CompareDifferentIdsAlphabetically<SearchCase>(GetComparisonIndex(), Other.GetComparisonIndex());
 #else	// UE_FNAME_OUTLINE_NUMBER
 	// Names match, check whether numbers match.
 	if (GetComparisonIndex() == Other.GetComparisonIndex())
@@ -3530,10 +3531,20 @@ int32 FName::Compare( const FName& Other ) const
 	}
 
 	// Names don't match. This means we don't even need to check numbers.
-	return CompareDifferentIdsAlphabetically<ESearchCase::IgnoreCase>(GetComparisonIndex(), Other.GetComparisonIndex());
+	return CompareDifferentIdsAlphabetically<SearchCase>(GetComparisonIndex(), Other.GetComparisonIndex());
 #endif // UE_FNAME_OUTLINE_NUMBER
 }
 
+int32 FName::Compare( const FName& Other ) const
+{
+	return CompareInternal<ESearchCase::IgnoreCase>(Other);
+}
+
+int32 FName::CompareSensitive( const FName& Other ) const
+{
+	return CompareInternal<ESearchCase::CaseSensitive>(Other);
+}
+
 uint32 FName::GetPlainNameString(TCHAR(&OutName)[NAME_SIZE]) const
 {
 	const FNameEntry& Entry = *GetDisplayNameEntry();
diff --git a/Engine/Source/Runtime/Core/Public/UObject/NameTypes.h b/Engine/Source/Runtime/Core/Public/UObject/NameTypes.h
--- a/Engine/Source/Runtime/Core/Public/UObject/NameTypes.h
+++ b/Engine/Source/Runtime/Core/Public/UObject/NameTypes.h
@@ -930,6 +930,14 @@ public:
 	 */
 	[[nodiscard]] CORE_API int32 Compare( const FName& Other ) const;
 
+	/**
+	 * Compares name to passed in one. Sort is case-sensitive and alphabetical ascending.
+	 *
+	 * @param	Other	Name to compare this against
+	 * @return	< 0 is this < Other, 0 if this == Other, > 0 if this > Other
+	 */
+	[[nodiscard]] CORE_API int32 CompareSensitive( const FName& Other ) const;
+
 	/**
 	 * Fast non-alphabetical order that is only stable during this process' lifetime.
 	 *
@@ -1260,6 +1268,9 @@ private:
 	template <typename StringBufferType>
 	FORCEINLINE void AppendStringInternal(StringBufferType& Out) const;
 
+	template <ESearchCase::Type SearchCase>
+	FORCEINLINE int32 CompareInternal(const FName& Other) const;
+
 	FORCEINLINE FNameEntryId GetDisplayIndexFast() const
 	{
 #if WITH_CASE_PRESERVING_NAME
