Fix unreachable code warnings
#rnx

diff --git a/Engine/Source/Runtime/Core/Private/Apple/AppleMemoryObject.cpp b/Engine/Source/Runtime/Core/Private/Apple/AppleMemoryObject.cpp
--- a/Engine/Source/Runtime/Core/Private/Apple/AppleMemoryObject.cpp
+++ b/Engine/Source/Runtime/Core/Private/Apple/AppleMemoryObject.cpp
@@ -62,8 +62,10 @@ NS_ASSUME_NONNULL_BEGIN
 {
 	// Denied!
 	return;
-	
+
+	PRAGMA_DISABLE_UNREACHABLE_CODE_WARNINGS	
 	[super dealloc];
+	PRAGMA_RESTORE_UNREACHABLE_CODE_WARNINGS	
 }
 
 - (nullable NSMethodSignature *)methodSignatureForSelector:(SEL)sel
@@ -142,9 +144,11 @@ NS_ASSUME_NONNULL_BEGIN
 	}
 	return;
 	
+	PRAGMA_DISABLE_UNREACHABLE_CODE_WARNINGS	
 	// Deliberately unreachable code to silence clang's error about not calling super - which in all other
 	// cases will be correct.
 	[super dealloc];
+	PRAGMA_RESTORE_UNREACHABLE_CODE_WARNINGS	
 }
 
 @end
diff --git a/Engine/Source/Runtime/Core/Private/Mac/MacMallocZone.cpp b/Engine/Source/Runtime/Core/Private/Mac/MacMallocZone.cpp
--- a/Engine/Source/Runtime/Core/Private/Mac/MacMallocZone.cpp
+++ b/Engine/Source/Runtime/Core/Private/Mac/MacMallocZone.cpp
@@ -185,6 +185,5 @@ bool FMacMallocCrashHandler::IsOnCrashedThread( void )
 	else
 	{
 		FPlatformProcess::SleepInfinite();
-		return false;
 	}
 }
