PVS 7.42

#rnx
#Synced-CL 52566575

diff --git a/Engine/Source/Runtime/Core/Public/Analysis/PvsStudioCodeAnalysis.h b/Engine/Source/Runtime/Core/Public/Analysis/PvsStudioCodeAnalysis.h
--- a/Engine/Source/Runtime/Core/Public/Analysis/PvsStudioCodeAnalysis.h
+++ b/Engine/Source/Runtime/Core/Public/Analysis/PvsStudioCodeAnalysis.h
@@ -75,4 +75,7 @@
 	// V1053: Calling the 'foo' virtual function in the constructor/destructor may lead to unexpected result at runtime.
 	//-V::1053
 
+	// 7.42 false positive
+	//-V:NAME_None:502
+
 #endif // #if defined( PVS_STUDIO )
