Fix unreachable code warnings
#rnx

diff --git a/Engine/Source/Runtime/Core/Private/Unix/UnixPlatformMemory.cpp b/Engine/Source/Runtime/Core/Private/Unix/UnixPlatformMemory.cpp
--- a/Engine/Source/Runtime/Core/Private/Unix/UnixPlatformMemory.cpp
+++ b/Engine/Source/Runtime/Core/Private/Unix/UnixPlatformMemory.cpp
@@ -489,7 +489,6 @@ void* FUnixPlatformMemory::BinnedAllocFromOS(SIZE_T Size)
 	{
 		FPlatformMemory::OnOutOfMemory(ActualSizeMapped, ExpectedAlignment);
 		// unreachable
-		return nullptr;
 	}
 
 	SIZE_T Offset = (reinterpret_cast<SIZE_T>(Pointer) % ExpectedAlignment);
@@ -506,7 +505,6 @@ void* FUnixPlatformMemory::BinnedAllocFromOS(SIZE_T Size)
 		{
 			FPlatformMemory::OnOutOfMemory(SizeToNextAlignedPointer, ExpectedAlignment);
 			// unreachable
-			return nullptr;
 		}
 
 		// account for reduced mmaped size
@@ -532,7 +530,6 @@ void* FUnixPlatformMemory::BinnedAllocFromOS(SIZE_T Size)
 		{
 			FPlatformMemory::OnOutOfMemory(TailSize, ExpectedAlignment);
 			// unreachable
-			return nullptr;
 		}
 	}
 
