[AutoRTFM] Consolidate transactionally-safe shared mutex and shared-recursive mutex.

This CL unifies these two almost-the-same APIs into one templated class. Previously, the transactionally-safe SharedMutex and SharedRecursiveMutex APIs had gaps: `IsLocked`,`IsLockShared`, `TryLock` and `TryLockShared` were variously absent.

#Synced-CL 52411280

diff --git a/Engine/Source/Programs/AutoRTFMTests/Private/MutexTests.cpp b/Engine/Source/Programs/AutoRTFMTests/Private/MutexTests.cpp
--- a/Engine/Source/Programs/AutoRTFMTests/Private/MutexTests.cpp
+++ b/Engine/Source/Programs/AutoRTFMTests/Private/MutexTests.cpp
@@ -18,12 +18,6 @@
 #include <thread>
 #include <type_traits>
 
-template<typename T>
-concept CHasIsLockedMethod = requires(T Object) { { Object.IsLocked() } -> std::same_as<bool>; };
-
-template<typename T>
-concept CHasIsLockSharedMethod = requires(T Object) { { Object.IsLockShared() } -> std::same_as<bool>; };
-
 TEMPLATE_TEST_CASE("MutexTests", "", 
 	UE::FTransactionallySafeMutex,
 	UE::FTransactionallySafeRecursiveMutex,
@@ -49,13 +43,12 @@ TEMPLATE_TEST_CASE("MutexTests", "",
 		std::is_same_v<MutexType, UE::FTransactionallySafeSharedMutex>;
 
 	// Some mutex implementations do not currently support the IsLocked() method
-	constexpr bool bSupportsIsLocked = CHasIsLockedMethod<MutexType>;
-
-	// True if the mutex has a IsLockShared() method.
-	constexpr bool bSupportsIsLockShared = CHasIsLockSharedMethod<MutexType>;
+	constexpr bool bSupportsIsLocked = UE::Private::CAllowsIsLocked<MutexType>;
 
 	// True if the mutex supports recursive locking
-	constexpr bool bIsRecursive = std::is_same_v<MutexType, UE::FTransactionallySafeSharedRecursiveMutex>;
+	constexpr bool bIsRecursive =
+		std::is_same_v<MutexType, UE::FTransactionallySafeRecursiveMutex> ||
+		std::is_same_v<MutexType, UE::FTransactionallySafeSharedRecursiveMutex>;
 
 	// True if the mutex is shared
 	constexpr bool bIsShared =
@@ -76,7 +69,7 @@ TEMPLATE_TEST_CASE("MutexTests", "",
 			Mutex.Lock();
 		}
 	};
-	
+
 	auto RequireNotLocked = [](MutexType& Mutex)
 	{
 		if constexpr (bSupportsIsLocked)
@@ -92,19 +85,56 @@ TEMPLATE_TEST_CASE("MutexTests", "",
 		}
 	};
 
-	auto RequireLockShared = [](MutexType& Mutex)
+	auto RequireLockSharedState = [](MutexType& Mutex, bool bState)
 	{
-		if constexpr (bSupportsIsLockShared)
+		if constexpr (bIsShared)
 		{
-			REQUIRE(Mutex.IsLockShared());
+			REQUIRE(Mutex.IsLockShared() == bState);
+		}
+		else
+		{
+			FAIL("unsupported lock type");
 		}
 	};
-	
-	auto RequireNotSharedLocked = [](MutexType& Mutex)
+	auto RequireLockShared = [&RequireLockSharedState](MutexType& Mutex)
+	{
+		RequireLockSharedState(Mutex, true);
+	};
+	auto RequireNotSharedLocked = [&RequireLockSharedState](MutexType& Mutex)
+	{
+		RequireLockSharedState(Mutex, false);
+	};
+
+	auto TryLockShared = [](MutexType& Mutex, auto& Link) -> bool
 	{
-		if constexpr (bSupportsIsLockShared)
+		if constexpr (bIsShared && bIsRecursive)
+		{
+			return Mutex.TryLockShared(Link);
+		}
+		else if constexpr (bIsShared && !bIsRecursive)
 		{
-			REQUIRE(!Mutex.IsLockShared());
+			return Mutex.TryLockShared();
+		}
+		else
+		{
+			FAIL("unsupported lock type");
+			return false;
+		}
+	};
+
+	auto UnlockShared = [](MutexType& Mutex, auto& Link)
+	{
+		if constexpr (bIsShared && bIsRecursive)
+		{
+			Mutex.UnlockShared(Link);
+		}
+		else if constexpr (bIsShared && !bIsRecursive)
+		{
+			Mutex.UnlockShared();
+		}
+		else
+		{
+			FAIL("unsupported lock type");
 		}
 	};
 	
@@ -1288,6 +1318,179 @@ TEMPLATE_TEST_CASE("MutexTests", "",
 
 	REQUIRE(NumLocksUsed == UE::Private::TransactionallySafeMutexPool::GetNumUsed());
 
+	if constexpr (bIsShared)
+	{
+		SECTION("SharedLock()")
+		{
+			MutexType Mutex;
+			RequireNotLocked(Mutex);
+			RequireNotSharedLocked(Mutex);
+			{
+				UE::TSharedLock LockA(Mutex);
+				RequireNotLocked(Mutex);
+				RequireLockShared(Mutex);
+			}
+			RequireNotLocked(Mutex);
+			RequireNotSharedLocked(Mutex);
+		}
+	}
+
+	REQUIRE(NumLocksUsed == UE::Private::TransactionallySafeMutexPool::GetNumUsed());
+
+	if constexpr (bIsShared)
+	{
+		SECTION("TryLockShared, UnlockShared")
+		{
+			MutexType Mutex;
+			RequireNotLocked(Mutex);
+			RequireNotSharedLocked(Mutex);
+
+			UE::Core::Private::FSharedRecursiveMutexLink Link;
+			REQUIRE(TryLockShared(Mutex, Link));
+			RequireNotLocked(Mutex);
+			RequireLockShared(Mutex);
+
+			UnlockShared(Mutex, Link);
+			RequireNotLocked(Mutex);
+			RequireNotSharedLocked(Mutex);
+		}
+	}
+
+	REQUIRE(NumLocksUsed == UE::Private::TransactionallySafeMutexPool::GetNumUsed());
+
+	if constexpr (bIsShared)
+	{
+		SECTION("Commit(TryLockShared, UnlockShared)")
+		{
+			MutexType Mutex;
+
+			AutoRTFM::Testing::Commit([&]
+			{
+				RequireNotLocked(Mutex);
+				RequireNotSharedLocked(Mutex);
+
+				UE::Core::Private::FSharedRecursiveMutexLink Link;
+				REQUIRE(TryLockShared(Mutex, Link));
+				RequireLocked(Mutex);  // a transactional LockShared is pessimized to a Lock
+
+				UnlockShared(Mutex, Link);
+				RequireLocked(Mutex);  // a transactional lock is held until the transaction completes
+			});
+
+			RequireNotLocked(Mutex);
+			RequireNotSharedLocked(Mutex);
+		}
+	}
+
+	REQUIRE(NumLocksUsed == UE::Private::TransactionallySafeMutexPool::GetNumUsed());
+
+	if constexpr (bIsShared)
+	{
+		SECTION("Abort(TryLockShared)")
+		{
+			MutexType Mutex;
+
+			AutoRTFM::Testing::Abort([&]
+			{
+				RequireNotLocked(Mutex);
+				RequireNotSharedLocked(Mutex);
+
+				UE::Core::Private::FSharedRecursiveMutexLink Link;
+				REQUIRE(TryLockShared(Mutex, Link));
+				RequireLocked(Mutex);  // a transactional LockShared is pessimized to a Lock
+
+				AutoRTFM::AbortTransaction();
+			});
+
+			RequireNotLocked(Mutex);
+			RequireNotSharedLocked(Mutex);
+		}
+	}
+
+	REQUIRE(NumLocksUsed == UE::Private::TransactionallySafeMutexPool::GetNumUsed());
+
+	if constexpr (bIsShared)
+	{
+		SECTION("Lock, TryLockShared, Unlock")
+		{
+			MutexType Mutex;
+			RequireNotLocked(Mutex);
+			RequireNotSharedLocked(Mutex);
+
+			{
+				UE::TScopeLock Lock(Mutex);
+				RequireLocked(Mutex);
+
+				UE::Core::Private::FSharedRecursiveMutexLink Link;
+				REQUIRE(!TryLockShared(Mutex, Link));
+				RequireLocked(Mutex);
+			}
+
+			RequireNotLocked(Mutex);
+			RequireNotSharedLocked(Mutex);
+		}
+	}
+
+	REQUIRE(NumLocksUsed == UE::Private::TransactionallySafeMutexPool::GetNumUsed());
+
+	if constexpr (bIsShared)
+	{
+		SECTION("Commit(Lock, TryLockShared, Unlock)")
+		{
+			MutexType Mutex;
+
+			AutoRTFM::Testing::Commit([&]
+			{
+				RequireNotLocked(Mutex);
+				RequireNotSharedLocked(Mutex);
+
+				{
+					UE::TScopeLock Lock(Mutex);
+					RequireLocked(Mutex);
+
+					UE::Core::Private::FSharedRecursiveMutexLink Link;
+					REQUIRE(!TryLockShared(Mutex, Link));
+					RequireLocked(Mutex);
+				}
+			});
+
+			RequireNotLocked(Mutex);
+			RequireNotSharedLocked(Mutex);
+		}
+	}
+
+	REQUIRE(NumLocksUsed == UE::Private::TransactionallySafeMutexPool::GetNumUsed());
+
+	if constexpr (bIsShared)
+	{
+		SECTION("Abort(Lock, TryLockShared)")
+		{
+			MutexType Mutex;
+
+			AutoRTFM::Testing::Abort([&]
+			{
+				RequireNotLocked(Mutex);
+				RequireNotSharedLocked(Mutex);
+
+				{
+					UE::TScopeLock Lock(Mutex);
+					RequireLocked(Mutex);
+
+					UE::Core::Private::FSharedRecursiveMutexLink Link;
+					REQUIRE(!TryLockShared(Mutex, Link));
+					RequireLocked(Mutex);
+
+					AutoRTFM::AbortTransaction();
+				}
+			});
+
+			RequireNotLocked(Mutex);
+			RequireNotSharedLocked(Mutex);
+		}
+	}
+
+	REQUIRE(NumLocksUsed == UE::Private::TransactionallySafeMutexPool::GetNumUsed());
+
 	if constexpr (bIsShared && bIsRecursive)
 	{
 		SECTION("SharedLock(SharedLock())")
diff --git a/Engine/Source/Runtime/Core/Public/Async/TransactionallySafeSharedMutex.h b/Engine/Source/Runtime/Core/Public/Async/TransactionallySafeSharedMutex.h
--- a/Engine/Source/Runtime/Core/Public/Async/TransactionallySafeSharedMutex.h
+++ b/Engine/Source/Runtime/Core/Public/Async/TransactionallySafeSharedMutex.h
@@ -4,12 +4,9 @@
 
 #include "Async/SharedMutex.h"
 #include "AutoRTFM.h"
-#include "Templates/SharedPointerFwd.h"
 
 #if UE_AUTORTFM
 #include "Async/TransactionallySafeMutexPool.h"
-#include "AutoRTFM/OpenWrapper.h"
-#include "Templates/SharedPointer.h"
 #endif
 
 namespace UE
@@ -17,6 +14,41 @@ namespace UE
 
 #if UE_AUTORTFM
 
+namespace Private
+{
+
+template <typename MutexType, typename... ArgsType>
+concept CAllowsLockShared = requires(MutexType& M, ArgsType&&... Args)
+{
+	M.LockShared(std::forward<ArgsType>(Args)...);
+};
+
+template <typename MutexType, typename... ArgsType>
+concept CAllowsUnlockShared = requires(MutexType& M, ArgsType&&... Args)
+{
+	M.UnlockShared(std::forward<ArgsType>(Args)...);
+};
+
+template <typename MutexType, typename... ArgsType>
+concept CAllowsTryLockShared = requires(MutexType& M, ArgsType&&... Args)
+{
+	{ M.TryLockShared(std::forward<ArgsType>(Args)...) } -> std::same_as<bool>;
+};
+
+template <typename MutexType>
+concept CAllowsIsLocked = requires(MutexType& M)
+{
+	{ M.IsLocked() } -> std::same_as<bool>;
+};
+
+template <typename MutexType>
+concept CAllowsIsLockShared = requires(MutexType& M)
+{
+	{ M.IsLockShared() } -> std::same_as<bool>;
+};
+
+}  // namespace Private
+
 // A transactionally-safe shared lock that works in the following novel ways:
 // - In the open (non-transactional):
 //   - Take the lock like before. Simple!
@@ -45,13 +77,14 @@ namespace UE
 //   Unlock          -------------
 //   -----------     SharedLock      <- Invalid because the transaction can potentially observe side
 //                                      effects of the open-threads writes!
-struct FTransactionallySafeSharedMutexDefinition
+template <typename LockType>
+struct TTransactionallySafeSharedMutexDefinition
 {
-	FTransactionallySafeSharedMutexDefinition() : State(UE::Private::TransactionallySafeMutexPool::New<FState>())
+	TTransactionallySafeSharedMutexDefinition() : State(UE::Private::TransactionallySafeMutexPool::New<FState>())
 	{
 	}
 
-	~FTransactionallySafeSharedMutexDefinition()
+	~TTransactionallySafeSharedMutexDefinition()
 	{
 		// Normally, State should never be null, but test "Locked Mutex In Destructed Object" expects
 		// that we can memzero a mutex safely. This doesn't crash, but leaks the mutex's State pointer.
@@ -62,7 +95,10 @@ struct FTransactionallySafeSharedMutexDefinition
 		}
 	}
 
-	void LockShared()
+
+	template <typename... ArgsType>
+		requires UE::Private::CAllowsLockShared<LockType, ArgsType&&...>
+	void LockShared(ArgsType&&... Args)
 	{
 		if (!AutoRTFM::IsIdle())
 		{
@@ -71,12 +107,14 @@ struct FTransactionallySafeSharedMutexDefinition
 		}
 		else
 		{
-			State->Mutex.LockShared();
+			State->Mutex.LockShared(std::forward<ArgsType>(Args)...);
 			ensure(0 == State->TransactionalLockCount);
 		}
 	}
 
-	void UnlockShared()
+	template <typename... ArgsType>
+		requires UE::Private::CAllowsUnlockShared<LockType, ArgsType&&...>
+	void UnlockShared(ArgsType&&... Args)
 	{
 		if (!AutoRTFM::IsIdle())
 		{
@@ -86,7 +124,7 @@ struct FTransactionallySafeSharedMutexDefinition
 		else
 		{
 			ensure(0 == State->TransactionalLockCount);
-			State->Mutex.UnlockShared();
+			State->Mutex.UnlockShared(std::forward<ArgsType>(Args)...);
 		}
 	}
 
@@ -123,15 +161,49 @@ struct FTransactionallySafeSharedMutexDefinition
 			: TryLockWithinTransaction();
 	}
 
-	[[nodiscard]] bool TryLockShared()
+	template <typename... ArgsType>
+		requires UE::Private::CAllowsTryLockShared<LockType, ArgsType&&...>
+	[[nodiscard]] bool TryLockShared(ArgsType&&... Args)
 	{
 		return AutoRTFM::IsIdle()
-			? State->Mutex.TryLockShared()
+			? State->Mutex.TryLockShared(std::forward<ArgsType>(Args)...)
 			: TryLockWithinTransaction();
 	}
 
+	/**
+	 * This method may give surprising results and should be used with caution!
+	 *
+	 * - You cannot safely use this function to determine whether Lock() will block.
+	 *   You may be in race with another thread which is also trying to lock the mutex.
+	 *
+	 * - If a Mutex is locked by the AutoRTFM thread, it will not fully release the mutex
+	 *   until the transaction fully succeeds or rolls back. Because of this, IsLocked may
+	 *   return true even after the mutex has been fully unlocked.
+	 */
+	AUTORTFM_OPEN [[nodiscard]] bool IsLocked() const
+		requires UE::Private::CAllowsIsLocked<LockType>
+	{
+ 		return State->Mutex.IsLocked();
+	}
+
+	/**
+	 * This method may give surprising results and should be used with caution!
+	 *
+	 * - You cannot safely use this function to determine whether LockShared() will block.
+	 *   You may be in race with another thread which is also trying to lock the mutex.
+	 *
+	 * - If a Mutex is locked by the AutoRTFM thread, it will not fully release the mutex
+	 *   until the transaction fully succeeds or rolls back. Because of this, IsLockShared may
+	 *   return true even after the mutex has been fully unlocked.
+	 */
+	AUTORTFM_OPEN [[nodiscard]] bool IsLockShared() const
+		requires UE::Private::CAllowsIsLockShared<LockType>
+	{
+		return State->Mutex.IsLockShared();
+	}
+
 private:
-	UE_NONCOPYABLE(FTransactionallySafeSharedMutexDefinition)
+	UE_NONCOPYABLE(TTransactionallySafeSharedMutexDefinition)
 
 	void LockWithinTransaction()
 	{
@@ -219,7 +291,7 @@ private:
 
 	struct FState final
 	{
-		::UE::FSharedMutex Mutex;
+		LockType Mutex;
 		uint32 TransactionalLockCount = 0;
 		uint32 RefCount = 1;
 
@@ -283,7 +355,7 @@ private:
 	FState* State = nullptr;
 };
 
-using FTransactionallySafeSharedMutex = ::UE::FTransactionallySafeSharedMutexDefinition;
+using FTransactionallySafeSharedMutex = ::UE::TTransactionallySafeSharedMutexDefinition<::UE::FSharedMutex>;
 
 #else
 using FTransactionallySafeSharedMutex = ::UE::FSharedMutex;
diff --git a/Engine/Source/Runtime/Core/Public/Async/TransactionallySafeSharedRecursiveMutex.h b/Engine/Source/Runtime/Core/Public/Async/TransactionallySafeSharedRecursiveMutex.h
--- a/Engine/Source/Runtime/Core/Public/Async/TransactionallySafeSharedRecursiveMutex.h
+++ b/Engine/Source/Runtime/Core/Public/Async/TransactionallySafeSharedRecursiveMutex.h
@@ -4,11 +4,9 @@
 
 #include "Async/SharedRecursiveMutex.h"
 #include "AutoRTFM.h"
-#include "Templates/SharedPointerFwd.h"
 
 #if UE_AUTORTFM
-#include "AutoRTFM/OpenWrapper.h"
-#include "Templates/SharedPointer.h"
+#include "Async/TransactionallySafeSharedMutex.h"
 #endif
 
 namespace UE
@@ -16,251 +14,7 @@ namespace UE
 
 #if UE_AUTORTFM
 
-// A transactionally-safe shared lock that supports recursive locking and works in the following novel ways:
-// - In the open (non-transactional):
-//   - Take the lock like before. Simple!
-//   - Free the lock like before too.
-// - In the closed (transactional):
-//   - During locking we query `TransactionalLockCount`:
-//	   - 0 means we haven't taken the lock within our transaction nest and need to acquire the lock.
-//     - Otherwise we already have the lock (and are preventing non-transactional code seeing any
-//       modifications we've made while holding the lock), so just bump `TransactionalLockCount`.
-//     - We also register an on-abort handler to release the lock should we abort (but we need to
-//       query `TransactionalLockCount` even there because we could be aborting an inner transaction
-//       and the parent transaction still wants to have the lock held!).
-//   - During unlocking we defer doing the unlock until the transaction commits.
-//
-// Thus with this approach we will hold this lock for the *entirety* of the transactional nest should
-// we take the lock during the transaction, thus preventing non-transactional code from seeing any
-// modifications we should make.
-//
-// If we are within a transaction, we pessimise our shared lock to a full lock. Note: that it should
-// potentially be possible to have shared locks work correctly, but serious care will have to be taken to
-// ensure that we don't have:
-//   Open Thread     Closed Thread
-//   -----------     SharedLock
-//   -----------     SharedUnlock
-//   Lock            -------------
-//   Unlock          -------------
-//   -----------     SharedLock      <- Invalid because the transaction can potentially observe side
-//                                      effects of the open-threads writes!
-struct FTransactionallySafeSharedRecursiveMutexDefinition
-{
-	FTransactionallySafeSharedRecursiveMutexDefinition() : State(UE::Private::TransactionallySafeMutexPool::New<FState>())
-	{
-	}
-
-	~FTransactionallySafeSharedRecursiveMutexDefinition()
-	{
-		// Normally, State should never be null, but test "Locked Mutex In Destructed Object" expects
-		// that we can memzero a mutex safely. This doesn't crash, but leaks the mutex's State pointer.
-		if (State)
-		{
-			State->ReleaseOnCommit();
-			State = nullptr;
-		}
-	}
-
-	void LockShared(Core::Private::FSharedRecursiveMutexLink& Link)
-	{
-		if (!AutoRTFM::IsIdle())
-		{
-			// Transactionally pessimise LockShared -> Lock.
-			LockWithinTransaction();
-		}
-		else
-		{
-			State->Mutex.LockShared(Link);
-			ensure(0 == State->TransactionalLockCount);
-		}
-	}
-
-	void UnlockShared(Core::Private::FSharedRecursiveMutexLink& Link)
-	{
-		if (!AutoRTFM::IsIdle())
-		{
-			// Transactionally pessimise UnlockShared -> Unlock.
-			UnlockWithinTransaction();
-		}
-		else
-		{
-			ensure(0 == State->TransactionalLockCount);
-			State->Mutex.UnlockShared(Link);
-		}
-	}
-
-	void Lock()
-	{
-		if (!AutoRTFM::IsIdle())
-		{
-			LockWithinTransaction();
-		}
-		else
-		{
-			State->Mutex.Lock();
-			ensure(0 == State->TransactionalLockCount);
-		}
-	}
-
-	void Unlock()
-	{
-		if (!AutoRTFM::IsIdle())
-		{
-			UnlockWithinTransaction();
-		}
-		else
-		{
-			ensure(0 == State->TransactionalLockCount);
-			State->Mutex.Unlock();
-		}
-	}
-
-	/** 
-	 * This method may give surprising results and should be used with caution!
-	 * 
-	 * - You cannot safely use this function to determine whether Lock() will block.
-	 *   You may be in race with another thread which is also trying to lock the mutex.
-	 * 
-	 * - If a Mutex is locked by the AutoRTFM thread, it will not fully release the mutex
-	 *   until the transaction fully succeeds or rolls back. Because of this, IsLocked may
-	 *   return true even after the mutex has been fully unlocked.
-	 */
-	AUTORTFM_OPEN [[nodiscard]] bool IsLocked() const
-	{
-		return State->Mutex.IsLocked();
-	}
-
-	/** 
-	 * This method may give surprising results and should be used with caution!
-	 * 
-	 * - You cannot safely use this function to determine whether LockShared() will block.
-	 *   You may be in race with another thread which is also trying to lock the mutex.
-	 * 
-	 * - If a Mutex is locked by the AutoRTFM thread, it will not fully release the mutex
-	 *   until the transaction fully succeeds or rolls back. Because of this, IsLockedShared may
-	 *   return true even after the mutex has been fully unlocked.
-	 */
-	AUTORTFM_OPEN [[nodiscard]] bool IsLockShared() const
-	{
-		return State->Mutex.IsLockShared();
-	}
-
-private:
-	UE_NONCOPYABLE(FTransactionallySafeSharedRecursiveMutexDefinition)
-
-	void LockWithinTransaction()
-	{
-		if (AutoRTFM::IsClosed())
-		{
-			// We explicitly extend the lifetime of `State` here, in case `this` lives on the
-			// stack and ends up being overwritten before the on-abort is hit.
-			State->Acquire();
-
-			UE_AUTORTFM_OPEN
-			{
-				State->LockTransactionalLock();
-			};
-
-			UE_AUTORTFM_ONABORT(State = this->State)
-			{
-				State->UnlockTransactionalLock();
-			};
-
-			State->ReleaseOnCommit();
-		}
-		else
-		{
-			State->LockTransactionalLock();
-		}
-	}
-
-	void UnlockWithinTransaction()
-	{
-		if (AutoRTFM::IsClosed())
-		{
-			// We explicitly extend the lifetime of `State` here, in case `this` lives on the
-			// stack and ends up being overwritten before the on-commit is hit.
-			State->Acquire();
-
-			UE_AUTORTFM_ONCOMMIT(State = this->State)
-			{
-				State->UnlockTransactionalLock();
-				State->ReleaseImmediately();
-			};
-		}
-		else
-		{
-			State->UnlockTransactionalLock();
-		}
-	}
-
-	struct FState final
-	{
-		::UE::FSharedRecursiveMutex Mutex;
-		uint32 TransactionalLockCount = 0;
-		uint32 RefCount = 1;
-
-		FState() = default;
-		~FState()
-		{
-			checkSlow(RefCount == 0);
-			ensure(0 == TransactionalLockCount);
-		}
-
-		AUTORTFM_DISABLE void LockTransactionalLock()
-		{
-			// The transactional system which can increment TransactionalLockCount
-			// is always single-threaded, thus this is safe to check without atomicity.
-			if (0 == TransactionalLockCount)
-			{
-				Mutex.Lock();
-			}
-
-			TransactionalLockCount += 1;
-		}
-
-		AUTORTFM_DISABLE void UnlockTransactionalLock()
-		{
-			ensure(TransactionalLockCount > 0);
-
-			TransactionalLockCount -= 1;
-			if (TransactionalLockCount == 0)
-			{
-				Mutex.Unlock();
-			}
-		}
-
-		AUTORTFM_ENABLE void Acquire()
-		{
-			RefCount += 1;
-		}
-
-		AUTORTFM_ENABLE void ReleaseOnCommit()
-		{
-			checkSlow(RefCount > 0);
-
-			AutoRTFM::OnCommit([this]
-			{
-				ReleaseImmediately();
-			});
-		}
-
-		AUTORTFM_DISABLE void ReleaseImmediately()
-		{
-			checkSlow(RefCount > 0);
-
-			RefCount -= 1;
-			if (RefCount == 0)
-			{
-				UE::Private::TransactionallySafeMutexPool::Delete(this);
-			}
-		}
-	};
-
-	FState* State = nullptr;
-};
-
-using FTransactionallySafeSharedRecursiveMutex = ::UE::FTransactionallySafeSharedRecursiveMutexDefinition;
+using FTransactionallySafeSharedRecursiveMutex = ::UE::TTransactionallySafeSharedMutexDefinition<::UE::FSharedRecursiveMutex>;
 
 template <>
 class TSharedLock<FTransactionallySafeSharedRecursiveMutex> final
