Pointer Variants: Support init from TSharedPtr without SharedFromThis

[RN] TVariantPtr no longer requires non-UObject type to be TSharedFromThis, accepts TSharedPtr

#Synced-CL 52568272

diff --git a/Engine/Source/Runtime/Core/Public/Templates/PointerVariants.h b/Engine/Source/Runtime/Core/Public/Templates/PointerVariants.h
--- a/Engine/Source/Runtime/Core/Public/Templates/PointerVariants.h
+++ b/Engine/Source/Runtime/Core/Public/Templates/PointerVariants.h
@@ -2,6 +2,8 @@
 
 #pragma once
 
+#include "Concepts/ConvertibleTo.h"
+#include "Concepts/UObject.h"
 #include "Misc/TVariant.h"
 #include "Templates/Requires.h"
 #include "Templates/SharedPointer.h"
@@ -73,15 +75,12 @@ namespace UE::Core::Private
 		{
 		}
 
-		template <
-			typename DerivedType
-			UE_REQUIRES(std::is_convertible_v<DerivedType*, BaseType*> && (std::is_base_of_v<UObject, DerivedType> || IsDerivedFromSharedFromThis<DerivedType>()))
-		>
+		/** Constructs from a raw pointer, either UObject or shareable-from-this. */
+		template<typename DerivedType>
+		requires CConvertibleTo<DerivedType*, BaseType*> && (CUObject<DerivedType> || IsDerivedFromSharedFromThis<DerivedType>())
 		constexpr TPtrVariantBase(DerivedType* InDerived)
-		{   
-			// TODO:	Look for a way to move this "smart pointer selection" logic outside of the constructor; then we can remove "TVariant" as a dependency
-			//			The problem is that "DerivedType" is defined in the constructor; We don't want TWeakPtrVariant or TStrongPtrVariant to have "DerivedType" as part of their template args
-			if constexpr (std::is_base_of_v<UObject, DerivedType>)
+		{
+			if constexpr (CUObject<DerivedType>)
 			{
 				PtrVariant.template Emplace<FObjectPtrWrapper>(Cast<const UObject>(InDerived), Cast<BaseType>(InDerived));
 			}
@@ -98,6 +97,18 @@ namespace UE::Core::Private
 			}
 		}
 
+		/** Constructor for TSharedPtr (if strong) or TWeakPtr (if weak). */
+		constexpr TPtrVariantBase(SharedType&& InSmartPtr)
+		{
+			PtrVariant.template Emplace<SharedType>(MoveTemp(InSmartPtr));
+		}
+
+		/** Constructor for TSharedPtr (if strong) or TWeakPtr (if weak). */
+		constexpr TPtrVariantBase(const SharedType& InSmartPtr)
+		{
+			PtrVariant.template Emplace<SharedType>(InSmartPtr);
+		}
+
 		TPtrVariantBase& operator=(TYPE_OF_NULLPTR)
 		{
 			Reset();
@@ -152,7 +163,6 @@ namespace UE::Core::Private
 		>
 		friend bool operator==(const TPtrVariantBase<LhsType, bLhsStrength>& Lhs, RhsType* Rhs);
 
-	private:
 		const BaseType* GetRawPtrValue_Internal() const
 		{
 			return ::Visit([this](auto& StoredValue)
@@ -176,7 +186,6 @@ namespace UE::Core::Private
 					}, PtrVariant);
 		}
 
-	private:
 		TVariant<SharedType, FObjectPtrWrapper> PtrVariant;
 	};
 
@@ -245,6 +254,11 @@ struct TStrongPtrVariant : public UE::Core::Private::TPtrVariantBase<BaseType, t
 			}, Super::PtrVariant);
 	}
 
+	UE_REWRITE BaseType& operator*() const
+	{
+		return *Get();
+	}
+
 	TWeakPtrVariant<BaseType> ToWeakVariant() const
 	{
 		if (Super::IsValid())
@@ -305,4 +319,4 @@ struct TWeakPtrVariant : public UE::Core::Private::TPtrVariantBase<BaseType, fal
 			return TStrongPtrVariant<BaseType>();
 		}
 	}
-};
\ No newline at end of file
+};
