PlatformThreadLocalStorage: In the editor implementation, fatal assert rather than deadlocking if we run out of TLS handles. This replaces deadlocks during cook with a fatal assert. Cook processes should always fatal assert rather than deadlocking, if more graceful error handling is not possible.
#rnx
#Synced-CL 52568437

diff --git a/Engine/Source/Runtime/Core/Private/Windows/WindowsPlatformTLS.cpp b/Engine/Source/Runtime/Core/Private/Windows/WindowsPlatformTLS.cpp
--- a/Engine/Source/Runtime/Core/Private/Windows/WindowsPlatformTLS.cpp
+++ b/Engine/Source/Runtime/Core/Private/Windows/WindowsPlatformTLS.cpp
@@ -5,6 +5,8 @@
 #if WITH_EDITOR
 
 #include "Templates/UnrealTemplate.h"
+#include "Logging/LogMacros.h"
+#include "CoreGlobals.h"
 
 #include <AtomicQueue.h>
 
@@ -36,6 +38,16 @@ namespace WindowsPlatformTLS_Private
 			return Queue.pop() - 1;
 		}
 
+		bool TryAlloc(uint32& OutIndex)
+		{
+			if (!Queue.try_pop(OutIndex))
+			{
+				return false;
+			}
+			OutIndex -= 1;
+			return true;
+		}
+
 		void Free(uint32 Value)
 		{
 			Queue.push(Value + 1);
@@ -132,7 +144,15 @@ namespace WindowsPlatformTLS_Private
 
 uint32 FWindowsPlatformTLS::AllocTlsSlot()
 {
-	return WindowsPlatformTLS_Private::FSingleton::Get().Slots.Alloc();
+	uint32 Result;
+	if (!WindowsPlatformTLS_Private::FSingleton::Get().Slots.TryAlloc(Result))
+	{
+		UE_LOGF(LogCore, Fatal, "Maximum number of TlsSlots (WINDOWS_MAX_NUM_TLS_SLOTS == %u) exceeded when calling AllocTlsSlot.",
+			WindowsPlatformTLS_Private::MaxSlots);
+		// If the fatal log is off, just loop forever until we get one.
+		Result = WindowsPlatformTLS_Private::FSingleton::Get().Slots.Alloc();
+	}
+	return Result;
 }
 
 void FWindowsPlatformTLS::FreeTlsSlot(uint32 SlotIndex)
