[BinnedMalloc2] Add double-free checks to small allocations.

FMallocBinned2 has no double-free detection for small allocations on the TLS fast path.
The freed pointer is blindly pushed onto a per-thread free list, silently corrupting it.
Large allocations already have canary-based detection via FPoolInfo.
This change adds low-cost double-free detection for all small allocations, gated by UE_MB2_DOUBLE_FREE_DETECTION (default: off) and the CVar `MallocBinned2.EnableSmallDoubleFreeDetection`.

Double-free detection mechanism:
* When a small allocation is freed, a canary value is written at offset sizeof(FFreeBlock) (byte 16) in the freed block ??? past both the FBundleNode (TLS free list) and FFreeBlock (pool free list) header areas, so it survives both recycling paths. On subsequent frees, this location is checked first. If the canary is already present, a double free has occurred and a warning and callstack is printed.
* The canary is derived from the pointer address (Ptr ^ magic constant), making false positives from user data a ~2^-64 probability. The canary is explicitly cleared on all malloc paths (TLS, pool, and NoTLS).
* Bins too small to hold the canary after the FFreeBlock header (the 16-byte bin) are skipped via a runtime check. A static_assert verifies sizeof(FBundleNode) <= sizeof(FFreeBlock).

Performance cost (measured via Catch2 benchmarks, 3 runs each):
* Single malloc+free: within noise (~1ns, <2%)
* Batch 1024 alloc/free: +3-5% (~2-3ns per alloc/free pair)
* Large allocations and OS path: unaffected

The check is placed after the fork-support early return in Free() and FreeNoTLS() to avoid dirtying copy-on-write pages shared with the parent process.

#Synced-CL 52547192

diff --git a/Engine/Source/Runtime/Core/Private/HAL/MallocBinned2.cpp b/Engine/Source/Runtime/Core/Private/HAL/MallocBinned2.cpp
--- a/Engine/Source/Runtime/Core/Private/HAL/MallocBinned2.cpp
+++ b/Engine/Source/Runtime/Core/Private/HAL/MallocBinned2.cpp
@@ -1,6 +1,8 @@
 // Copyright Epic Games, Inc. All Rights Reserved.
 
 #include "HAL/MallocBinned2.h"
+#include "Logging/LogVerbosity.h"
+#include "Misc/AssertionMacros.h"
 #include "Templates/Function.h"
 #include "GenericPlatform/GenericPlatformProcess.h"
 #include "Stats/Stats.h"
@@ -30,6 +32,14 @@ PRAGMA_DISABLE_UNSAFE_TYPECAST_WARNINGS
 #	define GBinned2MoveOSFreesOffTimeCriticalThreads	UE_DEFAULT_GBinned2MoveOSFreesOffTimeCriticalThreads
 #endif
 
+#if UE_MB2_DOUBLE_FREE_DETECTION
+	bool GMallocBinned2EnableSmallDoubleFreeDetection = true;
+	static FAutoConsoleVariableRef GGMallocBinned2EnableSmallDoubleFreeDetectionCVar(
+		TEXT("MallocBinned2.EnableSmallDoubleFreeDetection"),
+		GMallocBinned2EnableSmallDoubleFreeDetection,
+		TEXT("Enables the double-free detector for small allocations in FBinnedMalloc2"));
+#endif 
+
 #if UE_MB2_ALLOCATOR_STATS
 	int64 Binned2HashMemory = 0;
 #endif
@@ -101,6 +111,47 @@ static FORCEINLINE bool IsSupportedSize(SIZE_T Size)
 	return bResult;
 }
 
+#if UE_MB2_DOUBLE_FREE_DETECTION
+// Canary stored after the FFreeBlock header area in freed bins to detect double frees.
+// Placed at sizeof(FFreeBlock) so it survives both FBundleNode (TLS) and FFreeBlock (pool) overwrites.
+// XOR with pointer address makes false positives from user data ~2^-64 probability.
+// Bins too small to hold the canary (the 16-byte bin) are not covered.
+FORCEINLINE void FMallocBinned2::SetAndValidateDoubleFreeCanary(void* Ptr, uint32 BinSize)
+{
+	if (GMallocBinned2EnableSmallDoubleFreeDetection && BinSize > sizeof(FFreeBlock))
+	{
+		static constexpr uintptr_t DoubleFreeCanaryMagic = 0xFE2DF3EEDEADF3EE;
+		const uintptr_t FreedMarker = (uintptr_t)Ptr ^ DoubleFreeCanaryMagic;
+		uintptr_t* CanarySlot = (uintptr_t*)((uint8*)Ptr + sizeof(FFreeBlock));
+		if (UNLIKELY(*CanarySlot == FreedMarker))
+		{
+			ReportSmallDoubleFree(Ptr);
+		}
+		*CanarySlot = FreedMarker;
+	}
+}
+
+FORCEINLINE void FMallocBinned2::ClearDoubleFreeCanary(void* Ptr, uint32 BinSize)
+{
+	if (GMallocBinned2EnableSmallDoubleFreeDetection && BinSize > sizeof(FFreeBlock))
+	{
+		uintptr_t* CanarySlot = (uintptr_t*)((uint8*)Ptr + sizeof(FFreeBlock));
+		*CanarySlot = 0;
+	}
+}
+
+FORCENOINLINE void FMallocBinned2::ReportSmallDoubleFree(void* Ptr)
+{
+#if DO_ENSURE
+	ensureAlwaysMsgf(false, TEXT("FMallocBinned2 Suspected double-free when freeing pointer: %p"), Ptr);
+#else
+	UE_LOGF(LogMemory, Warning, "FMallocBinned2 Suspected double-free when freeing pointer: %p", Ptr);
+	FDebug::DumpStackTraceToLog(ELogVerbosity::Warning);
+#endif
+}
+
+#endif
+
 FMallocBinned2::FPoolInfo::FPoolInfo()
 	: Taken(0)
 	, Canary(ECanary::Unassigned)
@@ -384,6 +435,9 @@ FMallocBinned2::FMallocBinned2()
 	
 	static_assert(sizeof(FFreeBlock)  <= SmallBinSizesInternal[0], "sizeof(FFreeBlock)  must be fit in smallest allocation size handled by FMallocBinned2.");
 	static_assert(sizeof(FBundleNode) <= SmallBinSizesInternal[0], "sizeof(FBundleNode) must fit in smallest allocation size handled by FMallocBinned2.");
+#if UE_MB2_DOUBLE_FREE_DETECTION
+	static_assert(sizeof(FBundleNode) <= sizeof(FFreeBlock), "FBundleNode must fit within FFreeBlock.");
+#endif
 
 	const FGenericPlatformMemoryConstants Constants = FPlatformMemory::GetConstants();
 	GMB2PageSize = Constants.BinnedPageSize;
@@ -543,6 +597,9 @@ void* FMallocBinned2::MallocExternalSmall(SIZE_T Size, uint32 Alignment)
 		{
 			if (void* Result = Lists->Malloc(PoolIndex))
 			{
+#if UE_MB2_DOUBLE_FREE_DETECTION
+				ClearDoubleFreeCanary(Result, PoolIndexToBinSize(PoolIndex));
+#endif
 #if UE_MB2_ALLOCATOR_STATS
 				const uint32 BinSize = PoolIndexToBinSize(PoolIndex);
 				Lists->AllocatedMemory += BinSize;
@@ -570,6 +627,9 @@ void* FMallocBinned2::MallocExternalSmall(SIZE_T Size, uint32 Alignment)
 	}
 
 	void* Result = Pool->AllocateBin();
+#if UE_MB2_DOUBLE_FREE_DETECTION
+	ClearDoubleFreeCanary(Result, Table.BinSize);
+#endif
 #if UE_MB2_ALLOCATOR_STATS
 	++Table.TotalUsedBins;
 	AllocatedSmallPoolMemory.fetch_add(Table.BinSize, std::memory_order_relaxed);
@@ -635,6 +695,9 @@ void* FMallocBinned2::MallocNoTLS(SIZE_T Size, uint32 Alignment)
 		}
 
 		void* Result = Pool->AllocateBin();
+#if UE_MB2_DOUBLE_FREE_DETECTION
+		ClearDoubleFreeCanary(Result, Table.BinSize);
+#endif
 #if UE_MB2_ALLOCATOR_STATS
 		++Table.TotalUsedBins;
 		AllocatedSmallPoolMemory.fetch_add(Table.BinSize, std::memory_order_relaxed);
@@ -811,13 +874,18 @@ void FMallocBinned2::Free(void* Ptr)
 #if BINNED2_FORK_SUPPORT
 		if (BasePtr->CanaryAndForkState != CurrentCanary)
 		{
-			// This page was allocated before we forked so we want to avoid dirtying it by writing a linked list into it 
+			// This page was allocated before we forked so we want to avoid dirtying it by writing a linked list into it
 			return;
 		}
 #endif
+
 		const uint32 BinSize = BasePtr->BinSize;
 		const uint32 PoolIndex = BasePtr->PoolIndex;
 
+#if UE_MB2_DOUBLE_FREE_DETECTION
+		SetAndValidateDoubleFreeCanary(Ptr, BinSize);
+#endif
+
 		FBundleNode* BundlesToRecycle = nullptr;
 		FPerThreadFreeBlockLists* Lists = GMallocBinnedPerThreadCaches ? FPerThreadFreeBlockLists::Get() : nullptr;
 		if (Lists)
@@ -881,11 +949,15 @@ void FMallocBinned2::FreeNoTLS(void* Ptr)
 #if BINNED2_FORK_SUPPORT
 		if (BasePtr->CanaryAndForkState != CurrentCanary)
 		{
-			// This page was allocated before we forked so we want to avoid dirtying it by writing a linked list into it 
+			// This page was allocated before we forked so we want to avoid dirtying it by writing a linked list into it
 			return;
 		}
 #endif
 
+#if UE_MB2_DOUBLE_FREE_DETECTION
+		SetAndValidateDoubleFreeCanary(Ptr, GetPoolHeaderFromPointer(Ptr)->BinSize);
+#endif
+
 		FBundleNode* BundlesToRecycle = (FBundleNode*)Ptr;
 		BundlesToRecycle->SetNextNodeInCurrentBundle(nullptr);
 
diff --git a/Engine/Source/Runtime/Core/Public/HAL/MallocBinned2.h b/Engine/Source/Runtime/Core/Public/HAL/MallocBinned2.h
--- a/Engine/Source/Runtime/Core/Public/HAL/MallocBinned2.h
+++ b/Engine/Source/Runtime/Core/Public/HAL/MallocBinned2.h
@@ -38,6 +38,13 @@
 
 #define UE_MB2_ALLOCATOR_STATS				UE_MBC_ALLOCATOR_STATS
 
+// Double-free detection for small pool allocations. Writes a pointer-derived canary into freed
+// bins and checks for it on subsequent frees. Near-zero overhead (same cache line as free list
+// writes). Can be disabled if profiling shows unexpected impact.
+#ifndef UE_MB2_DOUBLE_FREE_DETECTION
+#	define UE_MB2_DOUBLE_FREE_DETECTION		0
+#endif
+
 // Canary value used in FFreeBlock
 // A constant value unless we're compiled with fork support in which case there are two values identifying whether the page
 // was allocated pre- or post-fork
@@ -225,6 +232,12 @@ private:
 		return (FFreeBlock*)AlignDown(Ptr, UE_MB2_LARGE_ALLOC);
 	}
 
+#if UE_MB2_DOUBLE_FREE_DETECTION
+	static void SetAndValidateDoubleFreeCanary(void* Ptr, uint32 BinSize);
+	static void ClearDoubleFreeCanary(void* Ptr, uint32 BinSize);
+	static void ReportSmallDoubleFree(void* Ptr);
+#endif
+
 public:
 	FMallocBinned2();
 	virtual ~FMallocBinned2();
