Reduced FMallocBinned2::FFreeBlock from 16 bytes to 8 bytes
Added support for 16 bytes bins to double free detection routine

#Synced-CL 52439882

diff --git a/Engine/Source/Runtime/Core/Private/HAL/MallocBinned2.cpp b/Engine/Source/Runtime/Core/Private/HAL/MallocBinned2.cpp
--- a/Engine/Source/Runtime/Core/Private/HAL/MallocBinned2.cpp
+++ b/Engine/Source/Runtime/Core/Private/HAL/MallocBinned2.cpp
@@ -115,10 +115,9 @@ static FORCEINLINE bool IsSupportedSize(SIZE_T Size)
 // Canary stored after the FFreeBlock header area in freed bins to detect double frees.
 // Placed at sizeof(FFreeBlock) so it survives both FBundleNode (TLS) and FFreeBlock (pool) overwrites.
 // XOR with pointer address makes false positives from user data ~2^-64 probability.
-// Bins too small to hold the canary (the 16-byte bin) are not covered.
-FORCEINLINE void FMallocBinned2::SetAndValidateDoubleFreeCanary(void* Ptr, uint32 BinSize)
+FORCEINLINE void FMallocBinned2::SetAndValidateDoubleFreeCanary(void* Ptr)
 {
-	if (GMallocBinned2EnableSmallDoubleFreeDetection && BinSize > sizeof(FFreeBlock))
+	if (GMallocBinned2EnableSmallDoubleFreeDetection)
 	{
 		static constexpr uintptr_t DoubleFreeCanaryMagic = 0xFE2DF3EEDEADF3EE;
 		const uintptr_t FreedMarker = (uintptr_t)Ptr ^ DoubleFreeCanaryMagic;
@@ -131,9 +130,9 @@ FORCEINLINE void FMallocBinned2::SetAndValidateDoubleFreeCanary(void* Ptr, uint3
 	}
 }
 
-FORCEINLINE void FMallocBinned2::ClearDoubleFreeCanary(void* Ptr, uint32 BinSize)
+FORCEINLINE void FMallocBinned2::ClearDoubleFreeCanary(void* Ptr)
 {
-	if (GMallocBinned2EnableSmallDoubleFreeDetection && BinSize > sizeof(FFreeBlock))
+	if (GMallocBinned2EnableSmallDoubleFreeDetection)
 	{
 		uintptr_t* CanarySlot = (uintptr_t*)((uint8*)Ptr + sizeof(FFreeBlock));
 		*CanarySlot = 0;
@@ -213,11 +212,11 @@ bool FMallocBinned2::FPoolInfo::HasFreeBin() const
 	return FirstFreeBlock && FirstFreeBlock->GetNumFreeBins() != 0;
 }
 
-void* FMallocBinned2::FPoolInfo::AllocateBin()
+void* FMallocBinned2::FPoolInfo::AllocateBin(uint32 BinSize)
 {
 	check(HasFreeBin());
 	++Taken;
-	void* Result = FirstFreeBlock->AllocateBin();
+	void* Result = FirstFreeBlock->AllocateBin(BinSize);
 	ExhaustPoolIfNecessary();
 	return Result;
 }
@@ -268,7 +267,7 @@ void FMallocBinned2::FPoolInfo::ExhaustPoolIfNecessary()
 {
 	if (FirstFreeBlock->GetNumFreeBins() == 0)
 	{
-		FirstFreeBlock = FirstFreeBlock->NextFreeBlock;
+		FirstFreeBlock = FirstFreeBlock->GetNextFreeBlock();
 	}
 	check(!FirstFreeBlock || FirstFreeBlock->GetNumFreeBins() != 0);
 }
@@ -303,16 +302,13 @@ struct FMallocBinned2::Private
 			else
 			{
 				// If we are freeing memory in this pool it must have the current canary and not the pre-fork one. All caches should have been cleared when forking.
-				check(NodePool->FirstFreeBlock->CanaryAndForkState == EBlockCanary::Zero || NodePool->FirstFreeBlock->CanaryAndForkState == Allocator.CurrentCanary);
+				check(GetPoolHeaderFromPointer(NodePool->FirstFreeBlock)->CanaryAndForkState == EBlockCanary::Zero || GetPoolHeaderFromPointer(NodePool->FirstFreeBlock)->CanaryAndForkState == Allocator.CurrentCanary);
 			}
 
 			// Free a pooled allocation.
 			FFreeBlock* Free = (FFreeBlock*)Node;
 			Free->NumFreeBins = 1;
-			Free->NextFreeBlock = NodePool->FirstFreeBlock;
-			Free->BinSize = BinSize;
-			Free->CanaryAndForkState = Allocator.CurrentCanary;
-			Free->PoolIndex = InPoolIndex;
+			Free->SetNextFreeBlock(NodePool->FirstFreeBlock);
 			NodePool->FirstFreeBlock = Free;
 
 			UE_MBC_UPDATE_STATS(--Table.TotalUsedBins);
@@ -322,7 +318,7 @@ struct FMallocBinned2::Private
 			if (--NodePool->Taken == 0)
 			{
 				NodePool->SetCanary(FPoolInfo::ECanary::Unassigned, true, false);
-				FFreeBlock* BasePtrOfNode = GetPoolHeaderFromPointer(Node);
+				FPoolHeader* BasePtrOfNode = GetPoolHeaderFromPointer(Node);
 
 				// Free the OS memory.
 				NodePool->Unlink();
@@ -333,7 +329,7 @@ struct FMallocBinned2::Private
 
 #if UE_MB2_ALLOCATOR_STATS
 				int NumBinsInBlock = GMB2PageSize / BinSize;
-				if (NumBinsInBlock * BinSize + sizeof(FFreeBlock) > GMB2PageSize)
+				if (NumBinsInBlock * BinSize + sizeof(FPoolHeader) > GMB2PageSize)
 				{
 					NumBinsInBlock--;
 				}
@@ -406,21 +402,21 @@ FMallocBinned2::FPoolInfo& FMallocBinned2::FPoolList::PushNewPoolToFront(FMalloc
 		OutOfMemory(GMB2PageSize);
 	}
 
-	FFreeBlock* Free = new (FreePtr) FFreeBlock(GMB2PageSize, Table.BinSize, InPoolIndex, Allocator.CurrentCanary);
-	check(IsAligned(Free, GMB2PageSize));
+	FPoolHeader* Header = new (FreePtr) FPoolHeader(GMB2PageSize, Table.BinSize, InPoolIndex, Allocator.CurrentCanary);
+	check(IsAligned(Header, GMB2PageSize));
 
 #if UE_MB2_ALLOCATOR_STATS
-	Table.TotalAllocatedBins += Free->NumFreeBins;
+	Table.TotalAllocatedBins += Header->NumFreeBins;
 	Table.TotalAllocatedMem += GMB2PageSize;
 	AllocatedOSSmallPoolMemory.fetch_add((int64)GMB2PageSize, std::memory_order_relaxed);
 #endif
 
 	UE::TUniqueLock Lock(Allocator.ExternalAllocMutex);
 	// Create pool
-	FPoolInfo* Result = Internal::GetOrCreatePoolInfo(Allocator, Free, FPoolInfo::ECanary::FirstFreeBlockIsPtr);
+	FPoolInfo* Result = Internal::GetOrCreatePoolInfo(Allocator, Header, FPoolInfo::ECanary::FirstFreeBlockIsPtr);
 	Result->Link(Front);
 	Result->Taken          = 0;
-	Result->FirstFreeBlock = Free;
+	Result->FirstFreeBlock = Header;
 
 	return *Result;
 }
@@ -435,6 +431,7 @@ FMallocBinned2::FMallocBinned2()
 	
 	static_assert(sizeof(FFreeBlock)  <= SmallBinSizesInternal[0], "sizeof(FFreeBlock)  must be fit in smallest allocation size handled by FMallocBinned2.");
 	static_assert(sizeof(FBundleNode) <= SmallBinSizesInternal[0], "sizeof(FBundleNode) must fit in smallest allocation size handled by FMallocBinned2.");
+	static_assert(sizeof(FFreeBlock) == 8, "FFreeBlock should fit into 8 bytes");
 #if UE_MB2_DOUBLE_FREE_DETECTION
 	static_assert(sizeof(FBundleNode) <= sizeof(FFreeBlock), "FBundleNode must fit within FFreeBlock.");
 #endif
@@ -598,7 +595,7 @@ void* FMallocBinned2::MallocExternalSmall(SIZE_T Size, uint32 Alignment)
 			if (void* Result = Lists->Malloc(PoolIndex))
 			{
 #if UE_MB2_DOUBLE_FREE_DETECTION
-				ClearDoubleFreeCanary(Result, PoolIndexToBinSize(PoolIndex));
+				ClearDoubleFreeCanary(Result);
 #endif
 #if UE_MB2_ALLOCATOR_STATS
 				const uint32 BinSize = PoolIndexToBinSize(PoolIndex);
@@ -626,9 +623,9 @@ void* FMallocBinned2::MallocExternalSmall(SIZE_T Size, uint32 Alignment)
 		Pool = &Table.ActivePools.PushNewPoolToFront(*this, Table, PoolIndex);
 	}
 
-	void* Result = Pool->AllocateBin();
+	void* Result = Pool->AllocateBin(Table.BinSize);
 #if UE_MB2_DOUBLE_FREE_DETECTION
-	ClearDoubleFreeCanary(Result, Table.BinSize);
+	ClearDoubleFreeCanary(Result);
 #endif
 #if UE_MB2_ALLOCATOR_STATS
 	++Table.TotalUsedBins;
@@ -650,7 +647,7 @@ void* FMallocBinned2::MallocExternalSmall(SIZE_T Size, uint32 Alignment)
 				{
 					break;
 				}
-				Result = Pool->AllocateBin();
+				Result = Pool->AllocateBin(Table.BinSize);
 				UE_MBC_UPDATE_STATS(++Table.TotalUsedBins);
 			}
 		}
@@ -694,9 +691,9 @@ void* FMallocBinned2::MallocNoTLS(SIZE_T Size, uint32 Alignment)
 			Pool = &Table.ActivePools.PushNewPoolToFront(*this, Table, PoolIndex);
 		}
 
-		void* Result = Pool->AllocateBin();
+		void* Result = Pool->AllocateBin(Table.BinSize);
 #if UE_MB2_DOUBLE_FREE_DETECTION
-		ClearDoubleFreeCanary(Result, Table.BinSize);
+		ClearDoubleFreeCanary(Result);
 #endif
 #if UE_MB2_ALLOCATOR_STATS
 		++Table.TotalUsedBins;
@@ -778,13 +775,13 @@ void* FMallocBinned2::Realloc(void* Ptr, SIZE_T NewSize, uint32 Alignment)
 	{
 		check(Ptr); // null is 64k aligned so we should not be here
 		// Reallocate to a smaller/bigger pool if necessary
-		FFreeBlock* Free = GetPoolHeaderFromPointer(Ptr);
-		CanaryTest(Free);
-		const uint32 BinSize = Free->BinSize;
-		const uint32 PoolIndex = Free->PoolIndex;
+		FPoolHeader* Header = GetPoolHeaderFromPointer(Ptr);
+		CanaryTest(Header);
+		const uint32 BinSize = Header->BinSize;
+		const uint32 PoolIndex = Header->PoolIndex;
 #if BINNED2_FORK_SUPPORT
 		// If the canary is the pre-fork one, we should not allow this allocation to grow in-place to avoid copying a page from the parent process.
-		if (Free->CanaryAndForkState == CurrentCanary && 
+		if (Header->CanaryAndForkState == CurrentCanary &&
 #else
 		if( 
 #endif
@@ -868,22 +865,21 @@ void FMallocBinned2::Free(void* Ptr)
 	if (!IsOSAllocation(Ptr))
 	{
 		check(Ptr); // null is 64k aligned so we should not be here
-		FFreeBlock* BasePtr = GetPoolHeaderFromPointer(Ptr);
-		CanaryTest(BasePtr);
+		FPoolHeader* Header = GetPoolHeaderFromPointer(Ptr);
+		CanaryTest(Header);
 
 #if BINNED2_FORK_SUPPORT
-		if (BasePtr->CanaryAndForkState != CurrentCanary)
+		if (Header->CanaryAndForkState != CurrentCanary)
 		{
 			// This page was allocated before we forked so we want to avoid dirtying it by writing a linked list into it
 			return;
 		}
 #endif
-
-		const uint32 BinSize = BasePtr->BinSize;
-		const uint32 PoolIndex = BasePtr->PoolIndex;
+		const uint32 BinSize = Header->BinSize;
+		const uint32 PoolIndex = Header->PoolIndex;
 
 #if UE_MB2_DOUBLE_FREE_DETECTION
-		SetAndValidateDoubleFreeCanary(Ptr, BinSize);
+		SetAndValidateDoubleFreeCanary(Ptr);
 #endif
 
 		FBundleNode* BundlesToRecycle = nullptr;
@@ -943,11 +939,11 @@ void FMallocBinned2::FreeNoTLS(void* Ptr)
 	if (!IsOSAllocation(Ptr))
 	{
 		check(Ptr); // null is 64k aligned so we should not be here
-		FFreeBlock* BasePtr = GetPoolHeaderFromPointer(Ptr);
-		CanaryTest(BasePtr);
+		FPoolHeader* Header = GetPoolHeaderFromPointer(Ptr);
+		CanaryTest(Header);
 
 #if BINNED2_FORK_SUPPORT
-		if (BasePtr->CanaryAndForkState != CurrentCanary)
+		if (Header->CanaryAndForkState != CurrentCanary)
 		{
 			// This page was allocated before we forked so we want to avoid dirtying it by writing a linked list into it
 			return;
@@ -955,13 +951,13 @@ void FMallocBinned2::FreeNoTLS(void* Ptr)
 #endif
 
 #if UE_MB2_DOUBLE_FREE_DETECTION
-		SetAndValidateDoubleFreeCanary(Ptr, GetPoolHeaderFromPointer(Ptr)->BinSize);
+		SetAndValidateDoubleFreeCanary(Ptr);
 #endif
 
 		FBundleNode* BundlesToRecycle = (FBundleNode*)Ptr;
 		BundlesToRecycle->SetNextNodeInCurrentBundle(nullptr);
 
-		Private::FreeBundles(*this, BundlesToRecycle, BasePtr->PoolIndex);
+		Private::FreeBundles(*this, BundlesToRecycle, Header->PoolIndex);
 	}
 	else if (Ptr)
 	{
@@ -976,7 +972,7 @@ void FMallocBinned2::FPoolList::ValidateActivePools() const
 		const FPoolInfo* Pool = *PoolPtr;
 		check(Pool->PtrToPrevNext == PoolPtr);
 		check(Pool->FirstFreeBlock);
-		for (FFreeBlock* Free = Pool->FirstFreeBlock; Free; Free = Free->NextFreeBlock)
+		for (FFreeBlock* Free = Pool->FirstFreeBlock; Free; Free = Free->GetNextFreeBlock())
 		{
 			check(Free->GetNumFreeBins() > 0);
 		}
@@ -1037,12 +1033,12 @@ void FMallocBinned2::FlushCurrentThreadCacheInternal(bool bNewEpochOnly)
 	FMallocBinnedCommonUtils::FlushCurrentThreadCache(*this, bNewEpochOnly);
 }
 
-void FMallocBinned2::CanaryFail(const FFreeBlock* Block) const
+void FMallocBinned2::CanaryFail(const FPoolHeader* Header) const
 {
 #if BINNED2_FORK_SUPPORT
-	UE_LOGF(LogMemory, Fatal, "FMallocBinned2 Attempt to realloc an unrecognized block %p   canary == 0x%x != 0x%x or 0x%x ", (void*)Block, (int32)Block->CanaryAndForkState, CurrentCanary, OldCanary);
+	UE_LOGF(LogMemory, Fatal, "FMallocBinned2 Attempt to realloc an unrecognized block %p   canary == 0x%x != 0x%x or 0x%x ", (void*)Header, (int32)Header->CanaryAndForkState, CurrentCanary, OldCanary);
 #else
-	UE_LOGF(LogMemory, Fatal, "FMallocBinned2 Attempt to realloc an unrecognized block %p   canary == 0x%x != 0x%x", (void*)Block, (int32)Block->CanaryAndForkState, int(CurrentCanary));
+	UE_LOGF(LogMemory, Fatal, "FMallocBinned2 Attempt to realloc an unrecognized block %p   canary == 0x%x != 0x%x", (void*)Header, (int32)Header->CanaryAndForkState, int(CurrentCanary));
 #endif
 }
 
diff --git a/Engine/Source/Runtime/Core/Public/HAL/MallocBinned2.h b/Engine/Source/Runtime/Core/Public/HAL/MallocBinned2.h
--- a/Engine/Source/Runtime/Core/Public/HAL/MallocBinned2.h
+++ b/Engine/Source/Runtime/Core/Public/HAL/MallocBinned2.h
@@ -55,7 +55,7 @@ enum class EBlockCanary : uint8
 	PreFork = 0xb7,
 	PostFork = 0xca,
 #else
-	Value = 0xe3 
+	Value = 0xe3
 #endif
 };
 
@@ -92,7 +92,7 @@ public:
 		void SetCanary(ECanary ShouldBe, bool bPreexisting, bool bGuaranteedToBeNew);
 
 		bool HasFreeBin() const;
-		void* AllocateBin();
+		void* AllocateBin(uint32 BinSize);
 
 		SIZE_T GetOSRequestedBytes() const;
 		SIZE_T GetOsAllocatedBytes() const;
@@ -112,13 +112,9 @@ private:
 	/** Information about a piece of free memory. */
 	struct FFreeBlock
 	{
-		inline FFreeBlock(uint32 InPageSize, uint16 InBinSize, uint8 InPoolIndex, EBlockCanary InCanary)
-			: BinSize(InBinSize)
-			, PoolIndex(InPoolIndex)
-			, CanaryAndForkState(InCanary)
-			, NextFreeBlock(nullptr)
+		inline FFreeBlock(uint32 InPageSize, uint16 InBinSize)
+			: NextFreeBlock(0)
 		{
-			check(InPoolIndex < MAX_uint8 && InBinSize <= MAX_uint16);
 			NumFreeBins = InPageSize / InBinSize;
 			if (NumFreeBins * InBinSize + sizeof(FFreeBlock) > InPageSize)
 			{
@@ -132,7 +128,17 @@ private:
 			return NumFreeBins;
 		}
 
-		inline void* AllocateBin()
+		FORCEINLINE FFreeBlock* GetNextFreeBlock() const
+		{
+			return (FFreeBlock*)NextFreeBlock;
+		}
+
+		FORCEINLINE void SetNextFreeBlock(FFreeBlock* Next)
+		{
+			NextFreeBlock = uint64(Next);
+		}
+
+		inline void* AllocateBin(uint32 BinSize)
 		{
 			--NumFreeBins;
 			// this is a pointer to a whole 64KB block
@@ -145,16 +151,31 @@ private:
 			return (uint8*)this + (NumFreeBins * BinSize);
 		}
 
+		constexpr static uint32 MaxBitsForFreeBins = 32 - FMath::CountLeadingZeros((UE_MB2_LARGE_ALLOC - 16) / 16);	// max number of free bins in the 64kb page - 4095, fits into 12 bits
+		static_assert(MaxBitsForFreeBins <= 16, "NextFreeBlockIndex should fit in 16 bits");
+
+		uint64 NumFreeBins : 16;		// Number of consecutive free bins here, at least 1
+		uint64 NextFreeBlock : 48;		// Next free block in another pool
+	};
+
+	struct FPoolHeader : public FFreeBlock
+	{
+		inline FPoolHeader(uint32 InPageSize, uint16 InBinSize, uint8 InPoolIndex, EBlockCanary InCanary)
+			: FFreeBlock(InPageSize, InBinSize)
+			, BinSize(InBinSize)
+			, PoolIndex(InPoolIndex)
+			, CanaryAndForkState(InCanary)
+		{
+			check(InPoolIndex < MAX_uint8 && InBinSize <= MAX_uint16);
+		}
+
 		uint16 BinSize;				// Size of the bins that this list points to
 		uint8 PoolIndex;			// Index of this pool
 
 		// Normally this value just functions as a canary to detect invalid memory state.
 		// When process forking is supported, it's still a canary but it has two valid values.
 		// One value is used pre-fork and one post-fork and the value is used to avoid freeing memory in pages shared with the parent process.
-		EBlockCanary CanaryAndForkState; 
-
-		uint32 NumFreeBins;          // Number of consecutive free bins here, at least 1.
-		FFreeBlock*  NextFreeBlock;  // Next free block in another pool
+		EBlockCanary CanaryAndForkState;
 	};
 
 	struct FPoolList
@@ -227,9 +248,9 @@ private:
 #endif
 	}
 
-	static FORCEINLINE FFreeBlock* GetPoolHeaderFromPointer(void* Ptr)
+	static FORCEINLINE FPoolHeader* GetPoolHeaderFromPointer(void* Ptr)
 	{
-		return (FFreeBlock*)AlignDown(Ptr, UE_MB2_LARGE_ALLOC);
+		return (FPoolHeader*)AlignDown(Ptr, UE_MB2_LARGE_ALLOC);
 	}
 
 #if UE_MB2_DOUBLE_FREE_DETECTION
@@ -262,9 +283,9 @@ public:
 	{
 		if (!IsOSAllocation(Ptr))
 		{
-			const FFreeBlock* Free = GetPoolHeaderFromPointer(Ptr);
-			CanaryTest(Free);
-			SizeOut = Free->BinSize;
+			const FPoolHeader* Header = GetPoolHeaderFromPointer(Ptr);
+			CanaryTest(Header);
+			SizeOut = Header->BinSize;
 			return true;
 		}
 		return false;
@@ -304,17 +325,17 @@ public:
 	void* MallocExternalSmall(SIZE_T Size, uint32 Alignment);
 	void* MallocExternalLarge(SIZE_T Size, uint32 Alignment);
 
-	void CanaryFail(const FFreeBlock* Block) const;
-	inline void CanaryTest(const FFreeBlock* Block) const
+	void CanaryFail(const FPoolHeader* Header) const;
+	inline void CanaryTest(const FPoolHeader* Header) const
 	{
 #if BINNED2_FORK_SUPPORT
 		// When we support forking there are two valid canary values.
-		if (Block->CanaryAndForkState != CurrentCanary && Block->CanaryAndForkState != OldCanary)
+		if (Header->CanaryAndForkState != CurrentCanary && Header->CanaryAndForkState != OldCanary)
 #else
-		if (Block->CanaryAndForkState != CurrentCanary)
+		if (Header->CanaryAndForkState != CurrentCanary)
 #endif
 		{
-			CanaryFail(Block);
+			CanaryFail(Header);
 		}
 	}
 
