Android: Fix errors setting the thread scheduler for TPri_Lowest priority threads. We were passing a pthread_t instead of a ThreadID to sched_setscheduler, which was logging an error and causing these threads to get SCHED_NORMAL instead of SCHED_BATCH.

[RN] Fix errors setting the thread scheduler for TPri_Lowest priority threads
#Synced-CL 52551815

diff --git a/Engine/Source/Runtime/Core/Private/Android/AndroidPlatformProcess.cpp b/Engine/Source/Runtime/Core/Private/Android/AndroidPlatformProcess.cpp
--- a/Engine/Source/Runtime/Core/Private/Android/AndroidPlatformProcess.cpp
+++ b/Engine/Source/Runtime/Core/Private/Android/AndroidPlatformProcess.cpp
@@ -324,7 +324,7 @@ void FRunnableThreadAndroid::SetThreadPriority(pthread_t InThread, EThreadPriori
 	const bool bCanLog = ErrorLogLimit >= 0;
 	bool bErrorEncountered = false;
 
-	if (InitialPolicy != NewPolicy && (sched_setscheduler(InThread, NewPolicy, &Sched)!=0) && bCanLog)
+	if (InitialPolicy != NewPolicy && (sched_setscheduler(ThreadID, NewPolicy, &Sched)!=0) && bCanLog)
 	{
 		bErrorEncountered = true;
 		UE_LOGF(LogHAL, Error, "Failed to set %ls thread scheduler, tid %d from %d to %d (errno %d)", *GetThreadName(), ThreadID, InitialPolicy, NewPolicy, errno);
