Added support for non-power-of-two FName pool shard counts, to allow alleviating hash table load without needing to double the capacity.

diff --git a/Engine/Source/Runtime/Core/Private/UObject/UnrealNames.cpp b/Engine/Source/Runtime/Core/Private/UObject/UnrealNames.cpp
--- a/Engine/Source/Runtime/Core/Private/UObject/UnrealNames.cpp
+++ b/Engine/Source/Runtime/Core/Private/UObject/UnrealNames.cpp
@@ -909,17 +909,22 @@ private:
 // Contention exists in some tool scenarios, for instance between main thread
 // and asset data gatherer thread during editor startup.
 // At runtime the async loading thread contends with game thread during asset loading.
+#if defined(UE_FNAME_POOL_SHARD_COUNT_OVERRIDE)
+static constexpr uint32 FNamePoolShards = UE_FNAME_POOL_SHARD_COUNT_OVERRIDE;
+static constexpr uint32 FNamePoolShardBits = FPlatformMath::ConstExprCeilLogTwo(FNamePoolShards);
+#else
 #if WITH_CASE_PRESERVING_NAME
-constexpr uint32 FNamePoolShardBits = 10;
+static constexpr uint32 FNamePoolShardBits = 10;
 #else
-constexpr uint32 FNamePoolShardBits = 8;
+static constexpr uint32 FNamePoolShardBits = 8;
+#endif
+static constexpr uint32 FNamePoolShards = 1U << FNamePoolShardBits;
 #endif
 
 #if !defined(UE_FNAME_POOL_INITIAL_SLOT_BITS)
 #define UE_FNAME_POOL_INITIAL_SLOT_BITS 8
 #endif
 
-constexpr uint32 FNamePoolShards = 1 << FNamePoolShardBits;
 constexpr uint32 FNamePoolInitialSlotBits = UE_FNAME_POOL_INITIAL_SLOT_BITS;
 constexpr uint32 FNamePoolInitialSlotsPerShard = 1 << FNamePoolInitialSlotBits;
 
@@ -942,11 +947,17 @@ struct FNameHash
 	FNameEntryHeader EntryProbeHeader; // Helps cull equality checks when probing inspects entries
 
 	static constexpr uint64 AlgorithmId = 0xC1640000;
-	static constexpr uint32 ShardMask = FNamePoolShards - 1;
+	static constexpr uint32 ShardMask = (1U << FNamePoolShardBits) - 1U;
+	static_assert((ShardMask & FNameSlot::ProbeHashMask) == 0, "Probe hash mask overlaps with bits used to calculate shard index");
+
+	static uint32 GetShardIndexFromHiBits(uint32 HiBits)
+	{
+		return HiBits % FNamePoolShards;
+	}
 
 	static uint32 GetShardIndex(uint64 Hash)
 	{
-		return static_cast<uint32>(Hash >> 32) & ShardMask;
+		return GetShardIndexFromHiBits(static_cast<uint32>(Hash >> 32));
 	}
 
 	template<class CharType>
@@ -991,9 +1002,7 @@ struct FNameHash
 		// @see FNameSlot::Used()
 		uint32 IsNoneBit = IsNone << FNameSlot::ProbeHashShift;
 
-		static_assert((ShardMask & FNameSlot::ProbeHashMask) == 0, "Masks overlap");
-
-		ShardIndex = Hi & ShardMask;
+		ShardIndex = GetShardIndexFromHiBits(Hi);
 		UnmaskedSlotIndex = Lo;
 		SlotProbeHash = (Hi & FNameSlot::ProbeHashMask) | IsNoneBit;
 		EntryProbeHeader.Len = Len;
