Write initial and final thermal state to the CSV profiler as metadata.
Fix the IOSCPU/ThermalState stat. Just set it every frame and don't rely on the global value (which is using the wrong enum type and is only set when the temperature changes). There's no need to update when it changes, because the cost of querying is negligible anyway (<0.01ms)
#Synced-CL 52463319

diff --git a/Engine/Source/Runtime/Core/Private/ProfilingDebugging/CsvProfiler.cpp b/Engine/Source/Runtime/Core/Private/ProfilingDebugging/CsvProfiler.cpp
--- a/Engine/Source/Runtime/Core/Private/ProfilingDebugging/CsvProfiler.cpp
+++ b/Engine/Source/Runtime/Core/Private/ProfilingDebugging/CsvProfiler.cpp
@@ -4092,6 +4092,12 @@ void FCsvProfiler::BeginCaptureInternal(const FCsvCaptureCommand& CurrentCommand
 		SetMetadataInternal(TEXT("ProgramSizeMB"), *FString::SanitizeFloat((float)FPlatformMemory::GetProgramSize() / 1024.0f / 1024.0f));
 	}
 
+	const EDeviceThermalState ThermalState = FPlatformMisc::GetDeviceThermalState();
+	if (ThermalState != EDeviceThermalState::Unsupported)
+	{
+		SetNonPersistentMetadata(TEXT("ThermalState_Start"), *FString::FromInt((int32)ThermalState));
+	}
+
 	bNamedEventsWasEnabled = (GCycleStatsShouldEmitNamedEvents > 0);
 
 	GCsvStatCounts = !!CVarCsvStatCounts.GetValueOnGameThread();
@@ -4458,6 +4464,12 @@ TSharedFuture<FString> FCsvProfiler::EndCapture(FGraphEventRef EventToSignal)
 	// Fire before we copy the metadata so it gives other systems a chance to write any final information.
 	OnCSVProfileEndRequestedDelegate.Broadcast();
 
+	const EDeviceThermalState ThermalState = FPlatformMisc::GetDeviceThermalState();
+	if (ThermalState != EDeviceThermalState::Unsupported)
+	{
+		SetNonPersistentMetadata(TEXT("ThermalState_End"), *FString::FromInt((int32)ThermalState));
+	}
+
 	SetNonPersistentMetadata(TEXT("EndTimestamp"), *FString::Printf(TEXT("%lld"), FDateTime::UtcNow().ToUnixTimestamp()));
 	SetNonPersistentMetadata(TEXT("CaptureDuration"), *FString::SanitizeFloat(FPlatformTime::Seconds()-CaptureStartTime));
 
