Fix equality check oversight in TestEqual/TestNotEqual for FText

#Synced-CL 52400047

diff --git a/Engine/Source/Runtime/Core/Private/Misc/AutomationTest.cpp b/Engine/Source/Runtime/Core/Private/Misc/AutomationTest.cpp
--- a/Engine/Source/Runtime/Core/Private/Misc/AutomationTest.cpp
+++ b/Engine/Source/Runtime/Core/Private/Misc/AutomationTest.cpp
@@ -2194,7 +2194,7 @@ bool FAutomationTestBase::TestEqual(const TCHAR* What, FStringView Actual, FStri
 
 bool FAutomationTestBase::TestEqual(const TCHAR* What, const FText Actual, const FText Expected)
 {
-	if (!Actual.EqualToCaseIgnored(Actual))
+	if (!Actual.EqualToCaseIgnored(Expected))
 	{
 		AddError(FString::Printf(TEXT("Expected '%s' to be %s, but it was %s."), What, *Expected.ToString(), *Actual.ToString()), 1);
 		return false;
@@ -2270,7 +2270,7 @@ bool FAutomationTestBase::TestNotEqual(const TCHAR* What, const double Actual, c
 
 bool FAutomationTestBase::TestNotEqual(const TCHAR* What, const FText Actual, const FText Expected)
 {
-	if (Actual.EqualToCaseIgnored(Actual))
+	if (Actual.EqualToCaseIgnored(Expected))
 	{
 		AddError(FString::Printf(TEXT("Expected '%s' to differ from %s, but it was %s."), What, *Expected.ToString(), *Actual.ToString()), 1);
 		return false;
