Added FMulticastInlineDelegateProperty declarations and names.

#Synced-CL 52474350

diff --git a/Engine/Source/Runtime/Core/Public/UObject/NativeTypeToPropertyType.h b/Engine/Source/Runtime/Core/Public/UObject/NativeTypeToPropertyType.h
--- a/Engine/Source/Runtime/Core/Public/UObject/NativeTypeToPropertyType.h
+++ b/Engine/Source/Runtime/Core/Public/UObject/NativeTypeToPropertyType.h
@@ -82,7 +82,7 @@ template <>                                         struct TNativeTypeToProperty
 template <>                                         struct TNativeTypeToPropertyType<FText>                                 { using Type = FTextProperty; };
 template <typename T>                               struct TNativeTypeToPropertyType<TArray<T, FDefaultAllocator>>          { using Type = FArrayProperty; };
 template <UE::CDynamicDelegate T>                   struct TNativeTypeToPropertyType<T>                                     { using Type = FDelegateProperty; };
-template <UE::CDynamicMulticastDelegate T>          struct TNativeTypeToPropertyType<T>                                     { using Type = FMulticastDelegateProperty; };
+template <UE::CDynamicMulticastDelegate T>          struct TNativeTypeToPropertyType<T>                                     { using Type = FMulticastInlineDelegateProperty; };
 template <UE::CSparseDynamicMulticastDelegate T>    struct TNativeTypeToPropertyType<T>                                     { using Type = FMulticastSparseDelegateProperty; };
 template <typename InKeyType, typename InValueType> struct TNativeTypeToPropertyType<TMap<InKeyType, InValueType>>          { using Type = FMapProperty; };
 template <typename InElementType>                   struct TNativeTypeToPropertyType<TSet<InElementType>>                   { using Type = FSetProperty; };
diff --git a/Engine/Source/Runtime/Core/Public/UObject/PropertyTypeToFName.h b/Engine/Source/Runtime/Core/Public/UObject/PropertyTypeToFName.h
--- a/Engine/Source/Runtime/Core/Public/UObject/PropertyTypeToFName.h
+++ b/Engine/Source/Runtime/Core/Public/UObject/PropertyTypeToFName.h
@@ -34,7 +34,7 @@ template <> struct TPropertyTypeToFName<FStrProperty>                     { stat
 template <> struct TPropertyTypeToFName<FTextProperty>                    { static FName Get() { return NAME_TextProperty; } };
 template <> struct TPropertyTypeToFName<FArrayProperty>                   { static FName Get() { return NAME_ArrayProperty; } };
 template <> struct TPropertyTypeToFName<FDelegateProperty>                { static FName Get() { return NAME_DelegateProperty; } };
-template <> struct TPropertyTypeToFName<FMulticastDelegateProperty>       { static FName Get() { return NAME_MulticastDelegateProperty; } };
+template <> struct TPropertyTypeToFName<FMulticastInlineDelegateProperty> { static FName Get() { return NAME_MulticastInlineDelegateProperty; } };
 template <> struct TPropertyTypeToFName<FMulticastSparseDelegateProperty> { static FName Get() { return NAME_MulticastSparseDelegateProperty; } };
 template <> struct TPropertyTypeToFName<FMapProperty>                     { static FName Get() { return NAME_MapProperty; } };
 template <> struct TPropertyTypeToFName<FSetProperty>                     { static FName Get() { return NAME_SetProperty; } };
diff --git a/Engine/Source/Runtime/Core/Public/UObject/UObjectHierarchyFwd.h b/Engine/Source/Runtime/Core/Public/UObject/UObjectHierarchyFwd.h
--- a/Engine/Source/Runtime/Core/Public/UObject/UObjectHierarchyFwd.h
+++ b/Engine/Source/Runtime/Core/Public/UObject/UObjectHierarchyFwd.h
@@ -49,7 +49,8 @@ class		FTextProperty;
 class		FArrayProperty;
 class		FDelegateProperty;
 class		FMulticastDelegateProperty;
-class		FMulticastSparseDelegateProperty;
+class			FMulticastInlineDelegateProperty;
+class			FMulticastSparseDelegateProperty;
 class		FMapProperty;
 class		FSetProperty;
 class		FEnumProperty;
diff --git a/Engine/Source/Runtime/Core/Public/UObject/UnrealNames.inl b/Engine/Source/Runtime/Core/Public/UObject/UnrealNames.inl
--- a/Engine/Source/Runtime/Core/Public/UObject/UnrealNames.inl
+++ b/Engine/Source/Runtime/Core/Public/UObject/UnrealNames.inl
@@ -56,6 +56,7 @@ REGISTER_NAME(37,AnsiStrProperty)
 REGISTER_NAME(38,WeakObjectProperty)
 REGISTER_NAME(39,MulticastSparseDelegateProperty)
 REGISTER_NAME(40,FieldPathProperty)
+REGISTER_NAME(41,MulticastInlineDelegateProperty)
 
 
 // Special types.
