Fix for [FAppleAsyncTask ProcessAsyncTasks] to run multiple tasks from the queue and in fixed order

#Synced-CL 52493178

diff --git a/Engine/Source/Runtime/Core/Private/Apple/AppleAsyncTask.cpp b/Engine/Source/Runtime/Core/Private/Apple/AppleAsyncTask.cpp
--- a/Engine/Source/Runtime/Core/Private/Apple/AppleAsyncTask.cpp
+++ b/Engine/Source/Runtime/Core/Private/Apple/AppleAsyncTask.cpp
@@ -90,33 +90,37 @@ NSMutableArray* RunningTasks;
  */
 + (void)ProcessAsyncTasks
 {
-	FAppleAsyncTask* CurrentTask = nil;
-
-	// grab one out of the queue to process outside of the lock
-	@synchronized(RunningTasks)
+	while (true)
 	{
-		if ([RunningTasks count] > 0)
+		FAppleAsyncTask* CurrentTask = nil;
+
+		// Grab the first task from the queue
+		@synchronized(RunningTasks)
 		{
-			CurrentTask = [RunningTasks objectAtIndex:0];
-			[RunningTasks removeObjectAtIndex:0];
+			if ([RunningTasks count] > 0)
+			{
+				CurrentTask = [RunningTasks objectAtIndex:0];
+			}
 		}
-	}
 
-	if (CurrentTask != nil)
-	{
-		if ([CurrentTask CheckForCompletion])
+		if (CurrentTask == nil)
 		{
-			// release the object count
-			[CurrentTask release];
+			break;
 		}
-		else
+
+		// Check and run the task
+		if (![CurrentTask CheckForCompletion])
 		{
-			// if it's not done, but it back in the list
-			@synchronized(RunningTasks)
-			{
-				[RunningTasks addObject:CurrentTask];
-			}
+			// Task not ready yet, break and resume on next frame to preserve task order
+			break;
+		}
+
+		// Success, remove the task and release it
+		@synchronized(RunningTasks)
+		{
+			[RunningTasks removeObjectAtIndex:0];
 		}
+		[CurrentTask release];
 	}
 }
 
