WindowsPlatformTLS: Remove unnecessary and misleading statements after UE_LOGF(..., Fatal, ...).
#rnx
#Synced-CL 52610419

diff --git a/Engine/Source/Runtime/Core/Private/Windows/WindowsPlatformTLS.cpp b/Engine/Source/Runtime/Core/Private/Windows/WindowsPlatformTLS.cpp
--- a/Engine/Source/Runtime/Core/Private/Windows/WindowsPlatformTLS.cpp
+++ b/Engine/Source/Runtime/Core/Private/Windows/WindowsPlatformTLS.cpp
@@ -6,6 +6,7 @@
 
 #include "Templates/UnrealTemplate.h"
 #include "Logging/LogMacros.h"
+#include "Logging/StructuredLog.h"
 #include "CoreGlobals.h"
 
 #include <AtomicQueue.h>
@@ -145,14 +146,12 @@ namespace WindowsPlatformTLS_Private
 uint32 FWindowsPlatformTLS::AllocTlsSlot()
 {
 	uint32 Result;
-	if (!WindowsPlatformTLS_Private::FSingleton::Get().Slots.TryAlloc(Result))
+	if (WindowsPlatformTLS_Private::FSingleton::Get().Slots.TryAlloc(Result))
 	{
-		UE_LOGF(LogCore, Fatal, "Maximum number of TlsSlots (WINDOWS_MAX_NUM_TLS_SLOTS == %u) exceeded when calling AllocTlsSlot.",
-			WindowsPlatformTLS_Private::MaxSlots);
-		// If the fatal log is off, just loop forever until we get one.
-		Result = WindowsPlatformTLS_Private::FSingleton::Get().Slots.Alloc();
+		return Result;
 	}
-	return Result;
+	UE_LOGFMT(LogCore, Fatal, "Maximum number of TlsSlots (WINDOWS_MAX_NUM_TLS_SLOTS == {MaxSlots}) exceeded when calling AllocTlsSlot.",
+		WindowsPlatformTLS_Private::MaxSlots);
 }
 
 void FWindowsPlatformTLS::FreeTlsSlot(uint32 SlotIndex)
