- Changed the autoload module (in DelayedAutoRegister phases) to have Editor specific list that are loaded when GIsEditor is true

#Synced-CL 52482450

diff --git a/Engine/Config/BaseEngine.ini b/Engine/Config/BaseEngine.ini
--- a/Engine/Config/BaseEngine.ini
+++ b/Engine/Config/BaseEngine.ini
@@ -31,6 +31,8 @@ SaveExt=usa
 Port=7777
 
 [AutoLoadModulesByPhase]
+
+[AutoLoadModulesByPhase_Editor]
 ; matching phase with FDelayedAutoRegisterHelper GSDKInit in TargetPlatformManagerModule.cpp
 +FileSystemReady=TargetPlatform
 
diff --git a/Engine/Source/Runtime/Core/Private/Misc/DelayedAutoRegister.cpp b/Engine/Source/Runtime/Core/Private/Misc/DelayedAutoRegister.cpp
--- a/Engine/Source/Runtime/Core/Private/Misc/DelayedAutoRegister.cpp
+++ b/Engine/Source/Runtime/Core/Private/Misc/DelayedAutoRegister.cpp
@@ -103,10 +103,11 @@ void FDelayedAutoRegisterHelper::RunAndClearDelayedAutoRegisterDelegates(EDelaye
 			FirstPhase = EDelayedRegisterRunPhase::StartOfEnginePreInit;
 		}
 
-		for (EDelayedRegisterRunPhase IniPhase = FirstPhase; IniPhase <= LastPhase; IniPhase = (EDelayedRegisterRunPhase)((int)IniPhase + 1))
+		// Lambda to load modules from a specific config property
+		auto LoadModulesFromConfigProperty = [RunPhase](EDelayedRegisterRunPhase IniPhase, const TCHAR* ConfigPropertyName)
 		{
 			TArray<FString> ModuleNames;
-			GConfig->GetArray(TEXT("AutoLoadModulesByPhase"), LexToString(IniPhase), ModuleNames, GEngineIni);
+			GConfig->GetArray(ConfigPropertyName, LexToString(IniPhase), ModuleNames, GEngineIni);
 			for (const FString& ModuleName : ModuleNames)
 			{
 				UE_LOGF(LogInit, Verbose, "Auto-loading module %ls during startup phase %ls", *ModuleName, LexToString(RunPhase));
@@ -115,6 +116,18 @@ void FDelayedAutoRegisterHelper::RunAndClearDelayedAutoRegisterDelegates(EDelaye
 					UE_LOGF(LogInit, Log, "Failed to auto-load module %ls during startup phase %ls", *ModuleName, LexToString(RunPhase));
 				}
 			}
+		};
+
+		for (EDelayedRegisterRunPhase IniPhase = FirstPhase; IniPhase <= LastPhase; IniPhase = (EDelayedRegisterRunPhase)((int)IniPhase + 1))
+		{
+			// Load modules from the standard config property
+			LoadModulesFromConfigProperty(IniPhase, TEXT("AutoLoadModulesByPhase"));
+
+			// Load editor-specific modules when running in editor
+			if (GIsEditor)
+			{
+				LoadModulesFromConfigProperty(IniPhase, TEXT("AutoLoadModulesByPhase_Editor"));
+			}
 		}
 	}
 
