Fix symbol redefinition when compiling some win64 targets with clang
#rnx

diff --git a/Engine/Source/Runtime/Core/Public/Microsoft/MSVCPlatform.h b/Engine/Source/Runtime/Core/Public/Microsoft/MSVCPlatform.h
--- a/Engine/Source/Runtime/Core/Public/Microsoft/MSVCPlatform.h
+++ b/Engine/Source/Runtime/Core/Public/Microsoft/MSVCPlatform.h
@@ -14,8 +14,12 @@
 	#error "Compiler is expected to support fold expressions"
 #endif
 
+#ifndef PLATFORM_RETURN_ADDRESS
 #define PLATFORM_RETURN_ADDRESS()	        _ReturnAddress()
+#endif
+#ifndef PLATFORM_RETURN_ADDRESS_POINTER
 #define PLATFORM_RETURN_ADDRESS_POINTER()	_AddressOfReturnAddress()
+#endif
 
 // https://devblogs.microsoft.com/cppblog/improving-the-state-of-debug-performance-in-c/
 #if __has_cpp_attribute(msvc::intrinsic)
@@ -23,7 +27,9 @@
 #endif
 
 // Ensure we can use this builtin - seems to be present on Clang 9, GCC 11 and MSVC 19.26
+#ifndef PLATFORM_COMPILER_SUPPORTS_BUILTIN_BITCAST
 #define PLATFORM_COMPILER_SUPPORTS_BUILTIN_BITCAST (_MSC_VER >= 1926)
+#endif
 
 #ifdef __has_cpp_attribute
 	#if __has_cpp_attribute(msvc::lifetimebound)
@@ -31,7 +37,11 @@
 	#endif
 #endif
 
+#ifndef UE_NO_PROFILE_ATTRIBUTE
 #define UE_NO_PROFILE_ATTRIBUTE
+#endif
 
 // MSVC doesn't complain if we pass `nullptr` to memory operations.
+#ifndef UE_PLATFORM_REQUIRES_NONNULL_MEMOP_ARGS
 #define UE_PLATFORM_REQUIRES_NONNULL_MEMOP_ARGS 0
+#endif
