[MallocBinned2] Fix build when UE_MB2_DOUBLE_FREE_DETECTION is enabled.

52617945 removed the BinSize parameter from the double-free methods in the .cpp, but not the header.
Also fix a typo in the CVar description.

#Synced-CL 52677305

diff --git a/Engine/Source/Runtime/Core/Private/HAL/MallocBinned2.cpp b/Engine/Source/Runtime/Core/Private/HAL/MallocBinned2.cpp
--- a/Engine/Source/Runtime/Core/Private/HAL/MallocBinned2.cpp
+++ b/Engine/Source/Runtime/Core/Private/HAL/MallocBinned2.cpp
@@ -37,7 +37,7 @@ PRAGMA_DISABLE_UNSAFE_TYPECAST_WARNINGS
 	static FAutoConsoleVariableRef GGMallocBinned2EnableSmallDoubleFreeDetectionCVar(
 		TEXT("MallocBinned2.EnableSmallDoubleFreeDetection"),
 		GMallocBinned2EnableSmallDoubleFreeDetection,
-		TEXT("Enables the double-free detector for small allocations in FBinnedMalloc2"));
+		TEXT("Enables the double-free detector for small allocations in FMallocBinned2"));
 #endif 
 
 #if UE_MB2_ALLOCATOR_STATS
diff --git a/Engine/Source/Runtime/Core/Public/HAL/MallocBinned2.h b/Engine/Source/Runtime/Core/Public/HAL/MallocBinned2.h
--- a/Engine/Source/Runtime/Core/Public/HAL/MallocBinned2.h
+++ b/Engine/Source/Runtime/Core/Public/HAL/MallocBinned2.h
@@ -254,8 +254,8 @@ private:
 	}
 
 #if UE_MB2_DOUBLE_FREE_DETECTION
-	static void SetAndValidateDoubleFreeCanary(void* Ptr, uint32 BinSize);
-	static void ClearDoubleFreeCanary(void* Ptr, uint32 BinSize);
+	static void SetAndValidateDoubleFreeCanary(void* Ptr);
+	static void ClearDoubleFreeCanary(void* Ptr);
 	static void ReportSmallDoubleFree(void* Ptr);
 #endif
 
