Fix FName::CompareSensitive not always comparing case-sensitively.

#[fyi] Steve.Robb
#Synced-CL 52729677

#Synced-CL 52734312

diff --git a/Engine/Source/Runtime/Core/Private/UObject/UnrealNames.cpp b/Engine/Source/Runtime/Core/Private/UObject/UnrealNames.cpp
--- a/Engine/Source/Runtime/Core/Private/UObject/UnrealNames.cpp
+++ b/Engine/Source/Runtime/Core/Private/UObject/UnrealNames.cpp
@@ -3520,28 +3520,42 @@ template <ESearchCase::Type SearchCase>
 int32 FName::CompareInternal( const FName& Other ) const
 {
 #if UE_FNAME_OUTLINE_NUMBER
-	if (GetComparisonIndexInternal() == Other.GetComparisonIndexInternal())
+	if constexpr (SearchCase == ESearchCase::IgnoreCase)
 	{
-		return 0; // Identical
+		if (GetComparisonIndexInternal() == Other.GetComparisonIndexInternal())
+		{
+			return 0; // Identical
+		}
 	}
-	else if (GetComparisonIndex() == Other.GetComparisonIndex())
+	else
 	{
-		// Same string, number may be different 
-		return GetNumber() - Other.GetNumber();
+		if (GetDisplayIndexInternal() == Other.GetDisplayIndexInternal())
+		{
+			return 0; // Identical
+		}
+	}
+#endif // UE_FNAME_OUTLINE_NUMBER
+
+	FNameEntryId A, B;
+	if constexpr (SearchCase == ESearchCase::IgnoreCase)
+	{
+		A = GetComparisonIndex();
+		B = Other.GetComparisonIndex();
+	}
+	else
+	{
+		A = GetDisplayIndexFast();
+		B = Other.GetDisplayIndexFast();
 	}
 
-	// Names don't match. This means we don't even need to check numbers.
-	return CompareDifferentIdsAlphabetically<SearchCase>(GetComparisonIndex(), Other.GetComparisonIndex());
-#else	// UE_FNAME_OUTLINE_NUMBER
 	// Names match, check whether numbers match.
-	if (GetComparisonIndex() == Other.GetComparisonIndex())
+	if (A == B)
 	{
 		return GetNumber() - Other.GetNumber();
 	}
 
 	// Names don't match. This means we don't even need to check numbers.
-	return CompareDifferentIdsAlphabetically<SearchCase>(GetComparisonIndex(), Other.GetComparisonIndex());
-#endif // UE_FNAME_OUTLINE_NUMBER
+	return CompareDifferentIdsAlphabetically<SearchCase>(A, B);
 }
 
 int32 FName::Compare( const FName& Other ) const
