Fix TIoStatusOr::ValueOrDie constness

#Synced-CL 52669770

diff --git a/Engine/Source/Runtime/Core/Public/IO/IoStatus.h b/Engine/Source/Runtime/Core/Public/IO/IoStatus.h
--- a/Engine/Source/Runtime/Core/Public/IO/IoStatus.h
+++ b/Engine/Source/Runtime/Core/Public/IO/IoStatus.h
@@ -133,7 +133,8 @@ public:
 	const FIoStatus&	Status() const;
 	bool				IsOk() const;
 
-	const T&			ValueOrDie();
+	const T&			ValueOrDie() const;
+	T&					ValueOrDie();
 	T					ConsumeValueOrDie();
 
 	void				Reset();
@@ -185,7 +186,18 @@ void TIoStatusOr<T>::Reset()
 }
 
 template<typename T>
-const T& TIoStatusOr<T>::ValueOrDie()
+const T& TIoStatusOr<T>::ValueOrDie() const
+{
+	if (!StatusValue.IsOk())
+	{
+		StatusOrCrash(StatusValue);
+	}
+
+	return *Value.GetTypedPtr();
+}
+
+template<typename T>
+inline T& TIoStatusOr<T>::ValueOrDie()
 {
 	if (!StatusValue.IsOk())
 	{
