PR #14672: Fix typo: funtion -> function in UnrealMathUtility.h comment

#rnx

#Synced-CL 52722726

diff --git a/Engine/Source/Runtime/Core/Public/Math/UnrealMathUtility.h b/Engine/Source/Runtime/Core/Public/Math/UnrealMathUtility.h
--- a/Engine/Source/Runtime/Core/Public/Math/UnrealMathUtility.h
+++ b/Engine/Source/Runtime/Core/Public/Math/UnrealMathUtility.h
@@ -1017,7 +1017,7 @@ public:
 	 * - positive azimuth means enemy is on the right of crosshair. (negative means left).
 	 * - positive elevation means enemy is on top of crosshair, negative means below.
 	 *
-	 * @Note: 'Azimuth' (.X) sign is changed to represent left/right and not front/behind. front/behind is the funtion's return value.
+	 * @Note: 'Azimuth' (.X) sign is changed to represent left/right and not front/behind. front/behind is the function's return value.
 	 *
 	 * @param	OutDotDist	.X = 'Direction' dot AxisX relative to plane (AxisX,AxisZ). (== Cos(Azimuth))
 	 *						.Y = 'Direction' dot AxisX relative to plane (AxisX,AxisY). (== Sin(Elevation))
