[AdditionalCrashContext]
- Add CrashedThreadId that the caller might need to find the right additional crash context information to attach

#Synced-CL 52789517

diff --git a/Engine/Source/Runtime/Core/Private/GenericPlatform/GenericPlatformCrashContext.cpp b/Engine/Source/Runtime/Core/Private/GenericPlatform/GenericPlatformCrashContext.cpp
--- a/Engine/Source/Runtime/Core/Private/GenericPlatform/GenericPlatformCrashContext.cpp
+++ b/Engine/Source/Runtime/Core/Private/GenericPlatform/GenericPlatformCrashContext.cpp
@@ -364,6 +364,7 @@ FGameDataSetDelegate FGenericCrashContext::OnGameDataSet;
 
 #if WITH_ADDITIONAL_CRASH_CONTEXTS
 FAdditionalCrashContextDelegate FGenericCrashContext::AdditionalCrashContextDelegate;
+FAdditionalCrashContextDelegateEx FGenericCrashContext::AdditionalCrashContextDelegateEx;
 #endif //WITH_ADDITIONAL_CRASH_CONTEXTS
 
 namespace NCached
@@ -2001,20 +2002,25 @@ private:
 
 #endif //WITH_ADDITIONAL_CRASH_CONTEXTS 
 
-
 void FGenericCrashContext::DumpAdditionalContext(const TCHAR* CrashFolderAbsolute)
+{
+	FAdditionalCrashContextParams Params;
+	DumpAdditionalContext(CrashFolderAbsolute, Params);
+}
+
+void FGenericCrashContext::DumpAdditionalContext(const TCHAR* CrashFolderAbsolute, const FAdditionalCrashContextParams& Params)
 {
 #if WITH_ADDITIONAL_CRASH_CONTEXTS 
 	FCrashContextExtendedWriterImpl Writer(CrashFolderAbsolute);
 	if (!UE::IsGCLocked())
 	{
 		AdditionalCrashContextDelegate.Broadcast(Writer);
+		AdditionalCrashContextDelegateEx.Broadcast(Writer, Params);
 	}
 	FAdditionalCrashContextStack::ExecuteProviders(Writer);
 #endif
 }
 
-
 /**
  * Attempts to create the output report directory.
  */
diff --git a/Engine/Source/Runtime/Core/Private/IOS/IOSPlatformCrashContext.cpp b/Engine/Source/Runtime/Core/Private/IOS/IOSPlatformCrashContext.cpp
--- a/Engine/Source/Runtime/Core/Private/IOS/IOSPlatformCrashContext.cpp
+++ b/Engine/Source/Runtime/Core/Private/IOS/IOSPlatformCrashContext.cpp
@@ -105,7 +105,11 @@ void FIOSCrashContext::GenerateInfoInFolder(char const* const InfoFolder, bool b
         }
 		
 		// generate additional files from providers.
-		DumpAdditionalContext(*FString(CrashInfoFolder));
+		FAdditionalCrashContextParams Params =
+		{
+			.CrashedThreadId = GetCrashedThreadId()
+		};
+		DumpAdditionalContext(*FString(CrashInfoFolder), Params);
                 
         // generate "minidump" (Apple crash log format)
         FCStringAnsi::Strncpy(FilePath, CrashInfoFolder, PATH_MAX);
diff --git a/Engine/Source/Runtime/Core/Private/Mac/MacPlatformMisc.cpp b/Engine/Source/Runtime/Core/Private/Mac/MacPlatformMisc.cpp
--- a/Engine/Source/Runtime/Core/Private/Mac/MacPlatformMisc.cpp
+++ b/Engine/Source/Runtime/Core/Private/Mac/MacPlatformMisc.cpp
@@ -2257,7 +2257,11 @@ void FMacCrashContext::GenerateInfoInFolder(char const* const InfoFolder) const
 		char FilePath[PATH_MAX] = {};
 		
 		// generate additional files from providers.
-		DumpAdditionalContext(*FString(CrashInfoFolder));
+		FAdditionalCrashContextParams Params =
+		{
+			.CrashedThreadId = GetCrashedThreadId()
+		};
+		DumpAdditionalContext(*FString(CrashInfoFolder), Params);
 		
 		// generate "minidump" (Apple crash log format)
 		FCStringAnsi::Strncpy(FilePath, CrashInfoFolder, PATH_MAX);
diff --git a/Engine/Source/Runtime/Core/Private/Unix/UnixPlatformCrashContext.cpp b/Engine/Source/Runtime/Core/Private/Unix/UnixPlatformCrashContext.cpp
--- a/Engine/Source/Runtime/Core/Private/Unix/UnixPlatformCrashContext.cpp
+++ b/Engine/Source/Runtime/Core/Private/Unix/UnixPlatformCrashContext.cpp
@@ -717,7 +717,11 @@ void FUnixCrashContext::GenerateCrashInfoAndLaunchReporter() const
 		if (IFileManager::Get().MakeDirectory(*CrashInfoAbsolute, true))
 		{
 			// generate additional files from providers.
-			DumpAdditionalContext(*CrashInfoFolder);
+			FAdditionalCrashContextParams Params =
+			{
+				.CrashedThreadId = GetCrashedThreadId()
+			};
+			DumpAdditionalContext(*CrashInfoFolder, Params);
 
 			// generate "minidump"
 			GenerateReport(FPaths::Combine(*CrashInfoAbsolute, TEXT("Diagnostics.txt")));
diff --git a/Engine/Source/Runtime/Core/Private/Windows/WindowsPlatformCrashContext.cpp b/Engine/Source/Runtime/Core/Private/Windows/WindowsPlatformCrashContext.cpp
--- a/Engine/Source/Runtime/Core/Private/Windows/WindowsPlatformCrashContext.cpp
+++ b/Engine/Source/Runtime/Core/Private/Windows/WindowsPlatformCrashContext.cpp
@@ -701,13 +701,13 @@ static FCriticalSection GMonitorLock;
 /**
  * Guard against additional context callbacks crashing
  */
-void GuardedDumpAdditionalContext(const TCHAR* CrashDirectoryAbsolute)
+void GuardedDumpAdditionalContext(const TCHAR* CrashDirectoryAbsolute, const FAdditionalCrashContextParams& Params)
 {
 #if !PLATFORM_SEH_EXCEPTIONS_DISABLED
 	__try
 #endif
 	{
-		FGenericCrashContext::DumpAdditionalContext(CrashDirectoryAbsolute);
+		FGenericCrashContext::DumpAdditionalContext(CrashDirectoryAbsolute, Params);
 	}
 #if !PLATFORM_SEH_EXCEPTIONS_DISABLED
 	__except (TRUE)
@@ -902,7 +902,11 @@ int32 ReportCrashForMonitor(
 		FGenericCrashContext::GenerateFiles(*CrashDirectoryAbsolute);
 
 		// Dump additional context
-		GuardedDumpAdditionalContext(*CrashDirectoryAbsolute);
+		FAdditionalCrashContextParams Params =
+		{ 
+			.CrashedThreadId = SharedContext->CrashingThreadId
+		};
+		GuardedDumpAdditionalContext(*CrashDirectoryAbsolute, Params);
 	}
 
 	// Allow the monitor process to take window focus
@@ -1036,7 +1040,11 @@ int32 ReportCrashUsingCrashReportClient(FWindowsPlatformCrashContext& InContext,
 			InContext.CopyPlatformSpecificFiles(*CrashFolderAbsolute, (void*) ExceptionInfo);
 
 			// Dump additional context
-			GuardedDumpAdditionalContext(*CrashFolderAbsolute);
+			FAdditionalCrashContextParams Params =
+			{
+				.CrashedThreadId = InContext.GetCrashedThreadId()
+			};
+			GuardedDumpAdditionalContext(*CrashFolderAbsolute, Params);
 
 			// Copy the log file to output
 			FGenericCrashContext::DumpLog(CrashFolderAbsolute);
diff --git a/Engine/Source/Runtime/Core/Public/GenericPlatform/GenericPlatformCrashContext.h b/Engine/Source/Runtime/Core/Public/GenericPlatform/GenericPlatformCrashContext.h
--- a/Engine/Source/Runtime/Core/Public/GenericPlatform/GenericPlatformCrashContext.h
+++ b/Engine/Source/Runtime/Core/Public/GenericPlatform/GenericPlatformCrashContext.h
@@ -316,6 +316,14 @@ struct FCrashContextExtendedWriter
 /** Simple Delegate for additional crash context. */
 DECLARE_MULTICAST_DELEGATE_OneParam(FAdditionalCrashContextDelegate, FCrashContextExtendedWriter&);
 
+/** Informations the callback might need to gather additional crash informations */
+struct FAdditionalCrashContextParams
+{
+	uint32 CrashedThreadId = ~(uint32)0;
+};
+
+DECLARE_MULTICAST_DELEGATE_TwoParams(FAdditionalCrashContextDelegateEx, FCrashContextExtendedWriter&, const FAdditionalCrashContextParams& Params);
+
 #endif //WITH_ADDITIONAL_CRASH_CONTEXTS
 
 /** Delegates for engine and game data set / reset */
@@ -652,8 +660,12 @@ public:
 	CORE_API static FString DumpLog(const FString& CrashFolderAbsolute);
 
 	/** Collects additional crash context providers. See FAdditionalCrashContextStack. */
+	UE_DEPRECATED(5.8, "Use the overload with FAdditionalCrashContextParams")
 	CORE_API static void DumpAdditionalContext(const TCHAR* CrashFolderAbsolute);
 
+	/** Collects additional crash context providers. See FAdditionalCrashContextStack. */
+	CORE_API static void DumpAdditionalContext(const TCHAR* CrashFolderAbsolute, const FAdditionalCrashContextParams& Params);
+
 	/** Initializes a shared crash context from current state. Will not set all fields in Dst. */
 	CORE_API static void CopySharedCrashContext(FSharedCrashContext& Dst);
 
@@ -685,6 +697,12 @@ public:
 	{
 		return AdditionalCrashContextDelegate;
 	}
+
+	/** Delegate for additional crash context. */
+	static inline FAdditionalCrashContextDelegateEx& OnAdditionalCrashContextDelegateEx()
+	{
+		return AdditionalCrashContextDelegateEx;
+	}
 #endif //WITH_ADDITIONAL_CRASH_CONTEXTS
 
 	/** Sets the process id to that has crashed. On supported platforms this will analyze the given process rather than current. Default is current process. */
@@ -699,6 +717,12 @@ public:
 		CrashedThreadId = InId;
 	}
 
+	/** Returns the crashing thread id. */
+	uint32 GetCrashedThreadId() const
+	{
+		return CrashedThreadId;
+	}
+
 	/** Sets the number of stack frames to ignore when symbolicating from a minidump */
 	CORE_API void SetNumMinidumpFramesToIgnore(int32 InNumMinidumpFramesToIgnore);
 
@@ -881,6 +905,7 @@ private:
 #if WITH_ADDITIONAL_CRASH_CONTEXTS
 	/** Delegate for additional crash context. */
 	CORE_API static FAdditionalCrashContextDelegate AdditionalCrashContextDelegate;
+	CORE_API static FAdditionalCrashContextDelegateEx AdditionalCrashContextDelegateEx;
 #endif //WITH_ADDITIONAL_CRASH_CONTEXTS
 
 	/** The buffer used to store the crash's properties. */
